Add thumbnail previews for runtime virtual texture

diff --git a/Engine/Source/Editor/VirtualTexturingEditor/Private/RuntimeVirtualTextureThumbnailRenderer.cpp b/Engine/Source/Editor/VirtualTexturingEditor/Private/RuntimeVirtualTextureThumbnailRenderer.cpp
new file mode 100644
--- /dev/null
+++ b/Engine/Source/Editor/VirtualTexturingEditor/Private/RuntimeVirtualTextureThumbnailRenderer.cpp
@@ -0,0 +1,82 @@
+// Copyright 1998-2019 Epic Games, Inc. All Rights Reserved.
+
+#include "RuntimeVirtualTextureThumbnailRenderer.h"
+
+#include "Components/RuntimeVirtualTextureComponent.h"
+#include "VT/RuntimeVirtualTexture.h"
+#include "VT/RuntimeVirtualTextureRender.h"
+
+namespace
+{
+	/** Find a matching component for this URuntimeVirtualTexture. */
+	URuntimeVirtualTextureComponent* FindComponent(URuntimeVirtualTexture* RuntimeVirtualTexture)
+	{
+		for (TObjectIterator<URuntimeVirtualTextureComponent> It; It; ++It)
+		{
+			URuntimeVirtualTextureComponent* RuntimeVirtualTextureComponent = *It;
+			if (RuntimeVirtualTextureComponent->GetVirtualTexture() == RuntimeVirtualTexture)
+			{
+				return RuntimeVirtualTextureComponent;
+			}
+		}
+
+		return nullptr;
+	}
+}
+
+URuntimeVirtualTextureThumbnailRenderer::URuntimeVirtualTextureThumbnailRenderer(const FObjectInitializer& ObjectInitializer)
+	: Super(ObjectInitializer)
+{
+}
+
+bool URuntimeVirtualTextureThumbnailRenderer::CanVisualizeAsset(UObject* Object)
+{
+	URuntimeVirtualTexture* RuntimeVirtualTexture = Cast<URuntimeVirtualTexture>(Object);
+	if (RuntimeVirtualTexture->GetEnabled())
+	{
+		//todo[vt]: Support thumbnails for ERuntimeVirtualTextureMaterialType::WorldHeight (which requires render to PF_G16)
+		if (RuntimeVirtualTexture->GetMaterialType() != ERuntimeVirtualTextureMaterialType::WorldHeight)
+		{
+			// We need a matching URuntimeVirtualTextureComponent in a Scene to be able to render a thumbnail
+			URuntimeVirtualTextureComponent* RuntimeVirtualTextureComponent = FindComponent(RuntimeVirtualTexture);
+			FSceneInterface* Scene = RuntimeVirtualTextureComponent != nullptr ? RuntimeVirtualTextureComponent->GetScene() : nullptr;
+			return Scene != nullptr;
+		}
+	}
+
+	return false;
+}
+
+void URuntimeVirtualTextureThumbnailRenderer::Draw(UObject* Object, int32 X, int32 Y, uint32 Width, uint32 Height, FRenderTarget* RenderTarget, FCanvas* Canvas)
+{
+	URuntimeVirtualTexture* RuntimeVirtualTexture = Cast<URuntimeVirtualTexture>(Object);
+	URuntimeVirtualTextureComponent* RuntimeVirtualTextureComponent = FindComponent(RuntimeVirtualTexture);
+	FSceneInterface* Scene = RuntimeVirtualTextureComponent != nullptr ? RuntimeVirtualTextureComponent->GetScene() : nullptr;
+	check(Scene != nullptr);
+
+	const FBox2D DestRect = FBox2D(FVector2D(X, Y), FVector2D(Width, Height));
+	const FTransform Transform = RuntimeVirtualTextureComponent->GetVirtualTextureTransform();
+	const uint32 VirtualTextureSceneIndex = RuntimeVirtualTexture::GetRuntimeVirtualTextureSceneIndex_GameThread(RuntimeVirtualTextureComponent);
+	const ERuntimeVirtualTextureMaterialType MaterialType = RuntimeVirtualTexture->GetMaterialType();
+
+	FVTProducerDescription VTDesc;
+	RuntimeVirtualTexture->GetProducerDescription(VTDesc, Transform);
+	const int32 MaxLevel = (int32)FMath::CeilLogTwo(FMath::Max(VTDesc.BlockWidthInTiles, VTDesc.BlockHeightInTiles));
+
+	ENQUEUE_RENDER_COMMAND(BakeStreamingTextureTileCommand)(
+		[Scene, VirtualTextureSceneIndex, MaterialType, RenderTarget, DestRect, Transform, MaxLevel](FRHICommandListImmediate& RHICmdList)
+	{
+		RuntimeVirtualTexture::RenderPage(
+			RHICmdList,
+			Scene->GetRenderScene(),
+			1 << VirtualTextureSceneIndex,
+			MaterialType,
+			RenderTarget->GetRenderTargetTexture(), DestRect,
+			nullptr, DestRect,
+			Transform,
+			FBox2D(FVector2D(0, 0), FVector2D(1, 1)),
+			MaxLevel,
+			MaxLevel,
+			ERuntimeVirtualTextureDebugType::None);
+	});
+}
diff --git a/Engine/Source/Editor/VirtualTexturingEditor/Private/RuntimeVirtualTextureThumbnailRenderer.h b/Engine/Source/Editor/VirtualTexturingEditor/Private/RuntimeVirtualTextureThumbnailRenderer.h
new file mode 100644
--- /dev/null
+++ b/Engine/Source/Editor/VirtualTexturingEditor/Private/RuntimeVirtualTextureThumbnailRenderer.h
@@ -0,0 +1,18 @@
+// Copyright 1998-2019 Epic Games, Inc. All Rights Reserved.
+
+#pragma once
+
+#include "CoreMinimal.h"
+#include "ThumbnailRendering/ThumbnailRenderer.h"
+#include "RuntimeVirtualTextureThumbnailRenderer.generated.h"
+
+UCLASS(MinimalAPI)
+class URuntimeVirtualTextureThumbnailRenderer : public UThumbnailRenderer
+{
+	GENERATED_UCLASS_BODY()
+
+	//~ Begin UThumbnailRenderer Interface.
+	virtual bool CanVisualizeAsset(UObject* Object);
+	virtual void Draw(UObject* Object, int32 X, int32 Y, uint32 Width, uint32 Height, FRenderTarget*, FCanvas* Canvas) override;
+	//~ EndUThumbnailRenderer Interface.
+};
diff --git a/Engine/Source/Editor/VirtualTexturingEditor/Private/VirtualTexturingEditorModule.cpp b/Engine/Source/Editor/VirtualTexturingEditor/Private/VirtualTexturingEditorModule.cpp
--- a/Engine/Source/Editor/VirtualTexturingEditor/Private/VirtualTexturingEditorModule.cpp
+++ b/Engine/Source/Editor/VirtualTexturingEditor/Private/VirtualTexturingEditorModule.cpp
@@ -6,6 +6,9 @@
 #include "PropertyEditorModule.h"
 #include "RuntimeVirtualTextureAssetTypeActions.h"
 #include "RuntimeVirtualTextureDetailsCustomization.h"
+#include "RuntimeVirtualTextureThumbnailRenderer.h"
+#include "ThumbnailRendering/ThumbnailManager.h"
+#include "VT/RuntimeVirtualTexture.h"
 #include "VT/RuntimeVirtualTextureVolume.h"
 
 #define LOCTEXT_NAMESPACE "VirtualTexturingEditorModule"
@@ -23,6 +26,8 @@ void FVirtualTexturingEditorModule::StartupModule()
 
 	IPlacementModeModule& PlacementModeModule = IPlacementModeModule::Get();
 	PlacementModeModule.OnPlacementModeCategoryRefreshed().AddRaw(this, &FVirtualTexturingEditorModule::OnPlacementModeRefresh);
+
+	UThumbnailManager::Get().RegisterCustomRenderer(URuntimeVirtualTexture::StaticClass(), URuntimeVirtualTextureThumbnailRenderer::StaticClass());
 }
 
 void FVirtualTexturingEditorModule::ShutdownModule()
diff --git a/Engine/Source/Editor/VirtualTexturingEditor/VirtualTexturingEditor.Build.cs b/Engine/Source/Editor/VirtualTexturingEditor/VirtualTexturingEditor.Build.cs
--- a/Engine/Source/Editor/VirtualTexturingEditor/VirtualTexturingEditor.Build.cs
+++ b/Engine/Source/Editor/VirtualTexturingEditor/VirtualTexturingEditor.Build.cs
@@ -10,12 +10,13 @@ public class VirtualTexturingEditor : ModuleRules
 
         PrivateIncludePathModuleNames.AddRange(
             new string[] {
-                "AssetRegistry",
-                "AssetTools",
-                "ContentBrowser",
-                "DesktopPlatform",
-                "MainFrame",
-                "UnrealEd",
+				"AssetRegistry",
+				"AssetTools",
+				"ContentBrowser",
+				"DesktopPlatform",
+				"MainFrame",
+				"Renderer",
+				"UnrealEd",
             }
         );
 
@@ -33,6 +34,9 @@ public class VirtualTexturingEditor : ModuleRules
 				"InputCore",
 				"PlacementMode",
 				"PropertyEditor",
+				"RenderCore",
+				"Renderer",
+				"RHI",
 				"Slate",
 				"SlateCore",
 				"UnrealEd",
diff --git a/Engine/Source/Runtime/Engine/Classes/Components/RuntimeVirtualTextureComponent.h b/Engine/Source/Runtime/Engine/Classes/Components/RuntimeVirtualTextureComponent.h
--- a/Engine/Source/Runtime/Engine/Classes/Components/RuntimeVirtualTextureComponent.h
+++ b/Engine/Source/Runtime/Engine/Classes/Components/RuntimeVirtualTextureComponent.h
@@ -27,6 +27,9 @@ public:
 	/** Get the runtime virtual texture object set on this component */
 	URuntimeVirtualTexture* GetVirtualTexture() const { return VirtualTexture; }
 
+	/** Get the runtime virtual texture UV to World transform on this component. */
+	FTransform GetVirtualTextureTransform() const;
+
 #if WITH_EDITOR
 	/** Copy the rotation from BoundsSourceActor to this component. Called by our UI details customization. */
 	void SetRotation();
diff --git a/Engine/Source/Runtime/Engine/Private/Components/RuntimeVirtualTextureComponent.cpp b/Engine/Source/Runtime/Engine/Private/Components/RuntimeVirtualTextureComponent.cpp
--- a/Engine/Source/Runtime/Engine/Private/Components/RuntimeVirtualTextureComponent.cpp
+++ b/Engine/Source/Runtime/Engine/Private/Components/RuntimeVirtualTextureComponent.cpp
@@ -44,6 +44,12 @@ void URuntimeVirtualTextureComponent::DestroyRenderState_Concurrent()
 	Super::DestroyRenderState_Concurrent();
 }
 
+FTransform URuntimeVirtualTextureComponent::GetVirtualTextureTransform() const
+{
+	// Transform is based on bottom left of the URuntimeVirtualTextureComponent unit box (which is centered on the origin)
+	return FTransform(FVector(-0.5f, -0.5f, 0.f)) * GetComponentTransform();
+}
+
 #if WITH_EDITOR
 
 void URuntimeVirtualTextureComponent::SetRotation()
diff --git a/Engine/Source/Runtime/Renderer/Private/VT/RuntimeVirtualTextureRender.cpp b/Engine/Source/Runtime/Renderer/Private/VT/RuntimeVirtualTextureRender.cpp
--- a/Engine/Source/Runtime/Renderer/Private/VT/RuntimeVirtualTextureRender.cpp
+++ b/Engine/Source/Runtime/Renderer/Private/VT/RuntimeVirtualTextureRender.cpp
@@ -2,6 +2,7 @@
 
 #include "VT/RuntimeVirtualTextureRender.h"
 
+#include "Components/RuntimeVirtualTextureComponent.h"
 #include "GlobalShader.h"
 #include "GPUScene.h"
 #include "MaterialShader.h"
@@ -13,6 +14,7 @@
 #include "SceneRenderTargets.h"
 #include "ShaderBaseClasses.h"
 #include "VT/RuntimeVirtualTexture.h"
+#include "VT/RuntimeVirtualTextureSceneProxy.h"
 #include "MeshPassProcessor.inl"
 
 
@@ -762,4 +764,24 @@ namespace RuntimeVirtualTexture
 			RHICmdList.CopyTexture(GraphOutputTexture1->GetRenderTargetItem().ShaderResourceTexture->GetTexture2D(), OutputTexture1->GetTexture2D(), Info);
 		}
 	}
+
+
+	uint32 GetRuntimeVirtualTextureSceneIndex_GameThread(class URuntimeVirtualTextureComponent* InComponent)
+	{
+		int32 SceneIndex = 0;
+		ENQUEUE_RENDER_COMMAND(GetSceneIndexCommand)(
+			[&SceneIndex, InComponent](FRHICommandListImmediate& RHICmdList)
+		{
+			if (InComponent->GetScene() != nullptr)
+			{
+				FScene* Scene = InComponent->GetScene()->GetRenderScene();
+				if (Scene != nullptr && InComponent->SceneProxy != nullptr)
+				{
+					SceneIndex = Scene->GetRuntimeVirtualTextureSceneIndex(InComponent->SceneProxy->ProducerId);
+				}
+			}
+		});
+		FlushRenderingCommands();
+		return SceneIndex;
+	}
 }
diff --git a/Engine/Source/Runtime/Renderer/Public/VT/RuntimeVirtualTextureRender.h b/Engine/Source/Runtime/Renderer/Public/VT/RuntimeVirtualTextureRender.h
--- a/Engine/Source/Runtime/Renderer/Public/VT/RuntimeVirtualTextureRender.h
+++ b/Engine/Source/Runtime/Renderer/Public/VT/RuntimeVirtualTextureRender.h
@@ -9,9 +9,19 @@ enum class ERuntimeVirtualTextureMaterialType : uint8;
 class FRHICommandListImmediate;
 class FRHITexture2D;
 class FScene;
+class URuntimeVirtualTextureComponent;
 
 namespace RuntimeVirtualTexture
 {
+#if WITH_EDITOR
+	/**
+	 * Get the scene index of the FRuntimeVirtualTextureSceneProxy associated with a URuntimeVirtualTextureComponent.
+	 * This is needed when rendering runtime virtual texture pages in alternative contexts such as when building previews etc.
+	 * This function is slow because it needs to flush render commands.
+	 */
+	RENDERER_API uint32 GetRuntimeVirtualTextureSceneIndex_GameThread(URuntimeVirtualTextureComponent* InComponent);
+#endif
+
 	/**
 	 * Render a single page of a virtual texture with a given material.
 	 * todo[vt]: Likely to be more optimal to batch several pages at a time and share setup/visibility/render targets.
