Runtime virtual texture can contain an optional streaming virtual texture to produce low mips
Add support for baking this with a button on the virtual texture volume

diff --git a/Engine/Source/Editor/VirtualTexturingEditor/Classes/RuntimeVirtualTextureDetailsCustomization.h b/Engine/Source/Editor/VirtualTexturingEditor/Classes/RuntimeVirtualTextureDetailsCustomization.h
--- a/Engine/Source/Editor/VirtualTexturingEditor/Classes/RuntimeVirtualTextureDetailsCustomization.h
+++ b/Engine/Source/Editor/VirtualTexturingEditor/Classes/RuntimeVirtualTextureDetailsCustomization.h
@@ -29,7 +29,6 @@ private:
 	TSharedPtr<class STextBlock> SizeText;
 	TSharedPtr<class STextBlock> TileSizeText;
 	TSharedPtr<class STextBlock> TileBorderSizeText;
-	TSharedPtr<class STextBlock> RemoveLowMipsText;
 
 	TSharedPtr<class STextBlock> PageTableTextureMemoryText;
 	TSharedPtr<class STextBlock> PhysicalTextureMemoryText;
@@ -49,6 +48,10 @@ protected:
 	FReply SetRotation();
 	/** Callback for Copy Bounds button */
 	FReply SetTransformToBounds();
+	/** Callback for Build button */
+	FReply BuildStreamedMips();
+	/** Callback for Build Debug button */
+	FReply BuildLowMipsDebug();
 
 	//~ Begin IDetailCustomization Interface.
 	virtual void CustomizeDetails(IDetailLayoutBuilder& DetailBuilder) override;
diff --git a/Engine/Source/Editor/VirtualTexturingEditor/Private/RuntimeVirtualTextureBuild.cpp b/Engine/Source/Editor/VirtualTexturingEditor/Private/RuntimeVirtualTextureBuild.cpp
new file mode 100644
--- /dev/null
+++ b/Engine/Source/Editor/VirtualTexturingEditor/Private/RuntimeVirtualTextureBuild.cpp
@@ -0,0 +1,158 @@
+// Copyright 1998-2019 Epic Games, Inc. All Rights Reserved.
+
+#include "RuntimeVirtualTextureBuild.h"
+
+#include "Components/RuntimeVirtualTextureComponent.h"
+#include "Engine/TextureRenderTarget2D.h"
+#include "Misc/ScopedSlowTask.h"
+#include "VT/RuntimeVirtualTexture.h"
+#include "VT/RuntimeVirtualTextureRender.h"
+
+namespace RuntimeVirtualTexture
+{
+	bool HasStreamedMips(class URuntimeVirtualTextureComponent* InComponent)
+	{
+		if (InComponent == nullptr)
+		{
+			return false;
+		}
+
+		URuntimeVirtualTexture* RuntimeVirtualTexture = InComponent->GetVirtualTexture();
+		if (RuntimeVirtualTexture == nullptr)
+		{
+			return false;
+		}
+
+		if (RuntimeVirtualTexture->GetStreamLowMips() <= 0)
+		{
+			return false;
+		}
+
+		return true;
+	}
+
+	bool BuildStreamedMips(class URuntimeVirtualTextureComponent* InComponent, ERuntimeVirtualTextureDebugType DebugType)
+	{
+		if (!HasStreamedMips(InComponent))
+		{
+			return true;
+		}
+
+		URuntimeVirtualTexture* RuntimeVirtualTexture = InComponent->GetVirtualTexture();
+		FSceneInterface* Scene = InComponent->GetScene();
+		const uint32 VirtualTextureSceneIndex = RuntimeVirtualTexture::GetRuntimeVirtualTextureSceneIndex_GameThread(InComponent);
+		const FTransform Transform = InComponent->GetVirtualTextureTransform();
+
+		FVTProducerDescription VTDesc;
+		RuntimeVirtualTexture->GetProducerDescription(VTDesc, Transform);
+
+		const int32 TileSize = VTDesc.TileSize;
+		const int32 TileBorderSize = VTDesc.TileBorderSize;
+		const int32 TextureSizeX = VTDesc.WidthInBlocks * VTDesc.BlockWidthInTiles * TileSize;
+		const int32 TextureSizeY = VTDesc.HeightInBlocks * VTDesc.BlockHeightInTiles * TileSize;
+		const int32 MaxLevel = (int32)FMath::CeilLogTwo(FMath::Max(VTDesc.BlockWidthInTiles, VTDesc.BlockHeightInTiles));
+		const int32 RenderLevel = FMath::Max(MaxLevel - RuntimeVirtualTexture->GetStreamLowMips() + 1, 1);
+		const int32 StreamingTextureSizeX = FMath::Max(TileSize, TextureSizeX >> RenderLevel);
+		const int32 StreamingTextureSizeY = FMath::Max(TileSize, TextureSizeY >> RenderLevel);
+		const int32 NumTilesX = StreamingTextureSizeX / TileSize;
+		const int32 NumTilesY = StreamingTextureSizeY / TileSize;
+		const int32 LayerCount = RuntimeVirtualTexture->GetLayerCount();
+
+		//todo[vt]: Support streaming VT for ERuntimeVirtualTextureMaterialType::WorldHeight (which requires UTextureRenderTarget2D support for PF_G16)
+		const ERuntimeVirtualTextureMaterialType MaterialType = RuntimeVirtualTexture->GetMaterialType();
+		if (MaterialType == ERuntimeVirtualTextureMaterialType::WorldHeight)
+		{
+			return false;
+		}
+
+		// Spin up slow task UI
+		const float TaskWorkRender = NumTilesX * NumTilesY;
+		const float TaskWorkBuildBulkData = NumTilesX * NumTilesY / 2;
+		FScopedSlowTask Task(TaskWorkRender + TaskWorkBuildBulkData, FText::AsCultureInvariant(RuntimeVirtualTexture->GetName()));
+		Task.MakeDialog(true);
+
+		// Final pixels will contain data for each virtual texture layer in order
+		TArray64<FColor> FinalPixels;
+		FinalPixels.InsertUninitialized(0, StreamingTextureSizeX * StreamingTextureSizeY * LayerCount);
+
+		// Allocate render targets for rendering out the runtime virtual texture tiles
+		UTextureRenderTarget2D* RenderTarget[2] = { nullptr };
+		FRenderTarget* RenderTargetResource[2] = { nullptr };
+
+		for (int32 Layer = 0; Layer < LayerCount; Layer++)
+		{
+			RenderTarget[Layer] = NewObject<UTextureRenderTarget2D>();
+			RenderTarget[Layer]->AddToRoot();
+			RenderTarget[Layer]->ClearColor = FLinearColor(0.0f, 0.0f, 0.0f, 0.0f);
+			RenderTarget[Layer]->InitCustomFormat(TileSize, TileSize, PF_B8G8R8A8, false);
+			RenderTargetResource[Layer] = RenderTarget[Layer]->GameThread_GetRenderTargetResource();
+		}
+
+		// Run over all tiles and render/store each one
+		for (int32 Y = 0; Y < NumTilesY && !Task.ShouldCancel(); Y++)
+		{
+			for (int32 X = 0; X < NumTilesX && !Task.ShouldCancel(); X++)
+			{
+				Task.EnterProgressFrame();
+
+				// Render tile
+				const FBox2D TileBox = FBox2D(FVector2D(0, 0), FVector2D(TileSize, TileSize));
+				const FBox2D UVRange = FBox2D(
+					FVector2D((float)X / (float)NumTilesX, (float)Y / (float)NumTilesY),
+					FVector2D((float)(X + 1) / (float)NumTilesX, (float)(Y + 1) / (float)NumTilesY));
+
+				ENQUEUE_RENDER_COMMAND(BakeStreamingTextureTileCommand)(
+					[Scene, VirtualTextureSceneIndex, MaterialType, RenderTargetResource, TileBox, Transform, UVRange, RenderLevel, MaxLevel, DebugType](FRHICommandListImmediate& RHICmdList)
+				{
+					RuntimeVirtualTexture::RenderPage(
+						RHICmdList,
+						Scene->GetRenderScene(),
+						1 << VirtualTextureSceneIndex,
+						MaterialType,
+						RenderTargetResource[0] != nullptr ? RenderTargetResource[0]->GetRenderTargetTexture() : nullptr, TileBox,
+						RenderTargetResource[1] != nullptr ? RenderTargetResource[1]->GetRenderTargetTexture() : nullptr, TileBox,
+						Transform,
+						UVRange,
+						RenderLevel,
+						MaxLevel,
+						DebugType);
+				});
+
+				FlushRenderingCommands();
+
+				// Read back into final pixel data
+				for (int32 Layer = 0; Layer < LayerCount; Layer++)
+				{
+					TArray<FColor> TilePixels;
+					RenderTarget[Layer]->GameThread_GetRenderTargetResource()->ReadPixels(TilePixels);
+
+					for (int32 y = 0; y < TileSize; y++)
+					{
+						for (int32 x = 0; x < TileSize; x++)
+						{
+							FinalPixels[Layer * StreamingTextureSizeX * StreamingTextureSizeY + (Y * TileSize + y) * StreamingTextureSizeX + (X * TileSize + x)] = TilePixels[y * TileSize + x];
+						}
+					}
+				}
+			}
+		}
+
+		for (int32 Layer = 0; Layer < LayerCount; Layer++)
+		{
+			RenderTarget[Layer]->RemoveFromRoot();
+		}
+
+		if (Task.ShouldCancel())
+		{
+			return false;
+		}
+
+		// Place final pixel data into the runtime virtual texture
+		Task.EnterProgressFrame(TaskWorkBuildBulkData);
+		RuntimeVirtualTexture->Modify();
+		RuntimeVirtualTexture->InitializeStreamingTexture(StreamingTextureSizeX, StreamingTextureSizeY, (uint8*)FinalPixels.GetData());
+		RuntimeVirtualTexture->PostEditChange();
+
+		return true;
+	}
+}
diff --git a/Engine/Source/Editor/VirtualTexturingEditor/Private/RuntimeVirtualTextureBuild.h b/Engine/Source/Editor/VirtualTexturingEditor/Private/RuntimeVirtualTextureBuild.h
new file mode 100644
--- /dev/null
+++ b/Engine/Source/Editor/VirtualTexturingEditor/Private/RuntimeVirtualTextureBuild.h
@@ -0,0 +1,14 @@
+// Copyright 1998-2019 Epic Games, Inc. All Rights Reserved.
+
+#pragma once
+
+#include "CoreMinimal.h"
+
+namespace RuntimeVirtualTexture
+{
+	/** Returns true if the component describes a runtime virtual texture that has streaming mips. */
+	bool HasStreamedMips(class URuntimeVirtualTextureComponent* InComponent);
+
+	/** Build the streaming mips and store in the component's associated URuntimeVirtualTexture object. */
+	bool BuildStreamedMips(class URuntimeVirtualTextureComponent* InComponent, enum class ERuntimeVirtualTextureDebugType DebugType);
+};
diff --git a/Engine/Source/Editor/VirtualTexturingEditor/Private/RuntimeVirtualTextureDetailsCustomization.cpp b/Engine/Source/Editor/VirtualTexturingEditor/Private/RuntimeVirtualTextureDetailsCustomization.cpp
--- a/Engine/Source/Editor/VirtualTexturingEditor/Private/RuntimeVirtualTextureDetailsCustomization.cpp
+++ b/Engine/Source/Editor/VirtualTexturingEditor/Private/RuntimeVirtualTextureDetailsCustomization.cpp
@@ -1,4 +1,5 @@
 // Copyright 1998-2019 Epic Games, Inc. All Rights Reserved.
+// Copyright 1998-2019 Epic Games, Inc. All Rights Reserved.
 
 #include "RuntimeVirtualTextureDetailsCustomization.h"
 
@@ -6,6 +7,7 @@
 #include "DetailCategoryBuilder.h"
 #include "DetailLayoutBuilder.h"
 #include "DetailWidgetRow.h"
+#include "RuntimeVirtualTextureBuild.h"
 #include "SResetToDefaultMenu.h"
 #include "VT/RuntimeVirtualTexture.h"
 #include "Widgets/Input/SButton.h"
@@ -95,7 +97,6 @@ void FRuntimeVirtualTextureDetailsCustomization::CustomizeDetails(IDetailLayoutB
 	AddTextToProperty(DetailBuilder, SizeCategory, "Size", SizeText);
 	AddTextToProperty(DetailBuilder, SizeCategory, "TileSize", TileSizeText);
 	AddTextToProperty(DetailBuilder, SizeCategory, "TileBorderSize", TileBorderSizeText);
-	AddTextToProperty(DetailBuilder, SizeCategory, "RemoveLowMips", RemoveLowMipsText);
 
 	// Add details block
 	IDetailCategoryBuilder& DetailsCategory = DetailBuilder.EditCategory("Details", FText::GetEmpty(), ECategoryPriority::Important);
@@ -129,6 +130,7 @@ void FRuntimeVirtualTextureDetailsCustomization::CustomizeDetails(IDetailLayoutB
 	DetailBuilder.GetProperty(FName(TEXT("MaterialType")))->SetOnPropertyValueChanged(FSimpleDelegate::CreateSP(this, &FRuntimeVirtualTextureDetailsCustomization::RefreshDetails));
 	DetailBuilder.GetProperty(FName(TEXT("bCompressTextures")))->SetOnPropertyValueChanged(FSimpleDelegate::CreateSP(this, &FRuntimeVirtualTextureDetailsCustomization::RefreshDetails));
 	DetailBuilder.GetProperty(FName(TEXT("RemoveLowMips")))->SetOnPropertyValueChanged(FSimpleDelegate::CreateSP(this, &FRuntimeVirtualTextureDetailsCustomization::RefreshDetails));
+	DetailBuilder.GetProperty(FName(TEXT("StreamLowMips")))->SetOnPropertyValueChanged(FSimpleDelegate::CreateSP(this, &FRuntimeVirtualTextureDetailsCustomization::RefreshDetails));
 
 	// Initialize text blocks
 	RefreshDetails();
@@ -143,7 +145,6 @@ void FRuntimeVirtualTextureDetailsCustomization::RefreshDetails()
  	SizeText->SetText(FText::Format(LOCTEXT("Details_Number", "{0}"), FText::AsNumber(VirtualTexture->GetSize(), &SizeOptions)));
  	TileSizeText->SetText(FText::Format(LOCTEXT("Details_Number", "{0}"), FText::AsNumber(VirtualTexture->GetTileSize(), &SizeOptions)));
  	TileBorderSizeText->SetText(FText::Format(LOCTEXT("Details_Number", "{0}"), FText::AsNumber(VirtualTexture->GetTileBorderSize(), &SizeOptions)));
- 	RemoveLowMipsText->SetText(FText::Format(LOCTEXT("Details_Number", "{0}"), FText::AsNumber(VirtualTexture->GetRemoveLowMips(), &SizeOptions)));
 
 	PageTableTextureMemoryText->SetText(FText::Format(LOCTEXT("Details_PageTableMemory", "Page Table Texture Memory (estimated): {0} KiB"), FText::AsNumber(VirtualTexture->GetEstimatedPageTableTextureMemoryKb(), &SizeOptions)));
 	PhysicalTextureMemoryText->SetText(FText::Format(LOCTEXT("Details_PhysicalMemory", "Physical Texture Memory (estimated): {0} KiB"), FText::AsNumber(VirtualTexture->GetEstimatedPhysicalTextureMemoryKb(), &SizeOptions)));
@@ -174,7 +175,7 @@ void FRuntimeVirtualTextureComponentDetailsCustomization::CustomizeDetails(IDeta
 		return;
 	}
 
-	// Use SourceActor property to add buttons
+	// Use existing property to add bounds copy buttons
 	TSharedPtr<IPropertyHandle> SourceActorValue = DetailBuilder.GetProperty("BoundsSourceActor");
 	DetailBuilder.HideProperty(SourceActorValue);
 
@@ -231,18 +232,81 @@ void FRuntimeVirtualTextureComponentDetailsCustomization::CustomizeDetails(IDeta
 			]
 		]
 	];
+
+	// Use existing property to add build button
+	TSharedPtr<IPropertyHandle> LowMipsValue = DetailBuilder.GetProperty("bUseStreamingLowMipsInEditor");
+	DetailBuilder.HideProperty(LowMipsValue);
+
+	IDetailCategoryBuilder& VirtualTextureCategory = DetailBuilder.EditCategory("VirtualTexture", FText::GetEmpty());
+	VirtualTextureCategory.AddCustomRow(LowMipsValue->GetPropertyDisplayName(), true)
+	.NameContent()
+	[
+		LowMipsValue->CreatePropertyNameWidget()
+	]
+	.ValueContent()
+	.MaxDesiredWidth(TOptional<float>())
+	[
+		SNew(SHorizontalBox)
+
+		+ SHorizontalBox::Slot()
+		[
+			LowMipsValue->CreatePropertyValueWidget()
+		]
+
+		+ SHorizontalBox::Slot()
+		[
+			SNew(SWrapBox)
+
+			+ SWrapBox::Slot()
+			[
+				SNew(SButton)
+				.Text(LOCTEXT("Button_Build", "Build"))
+				.ToolTipText(LOCTEXT("Button_Build_Tooltip", "Build the low mips as streaming virtual texture data"))
+				.OnClicked(this, &FRuntimeVirtualTextureComponentDetailsCustomization::BuildStreamedMips)
+			]
+		]
+
+		+ SHorizontalBox::Slot()
+		[
+			SNew(SWrapBox)
+
+			+ SWrapBox::Slot()
+			[
+				SNew(SButton)
+				.Text(LOCTEXT("Button_BuildDebig", "Build Debug"))
+				.ToolTipText(LOCTEXT("Button_BuildDebug_Tooltip", "Build the low mips with debug data"))
+				.OnClicked(this, &FRuntimeVirtualTextureComponentDetailsCustomization::BuildLowMipsDebug)
+			]
+		]
+	];
 }
 
 FReply FRuntimeVirtualTextureComponentDetailsCustomization::SetRotation()
 {
+	FScopedTransaction BakeTransaction(LOCTEXT("Transaction_CopyRotation", "Copy Rotation"));
+	RuntimeVirtualTextureComponent->Modify();
 	RuntimeVirtualTextureComponent->SetRotation();
 	return FReply::Handled();
 }
 
 FReply FRuntimeVirtualTextureComponentDetailsCustomization::SetTransformToBounds()
 {
+	FScopedTransaction BakeTransaction(LOCTEXT("Transaction_CopyBounds", "Copy Bounds"));
+	RuntimeVirtualTextureComponent->Modify();
 	RuntimeVirtualTextureComponent->SetTransformToBounds();
 	return FReply::Handled();
 }
 
+FReply FRuntimeVirtualTextureComponentDetailsCustomization::BuildStreamedMips()
+{
+	bool bOK = RuntimeVirtualTexture::BuildStreamedMips(RuntimeVirtualTextureComponent, ERuntimeVirtualTextureDebugType::None);
+	return bOK ? FReply::Handled() : FReply::Unhandled();
+}
+
+FReply FRuntimeVirtualTextureComponentDetailsCustomization::BuildLowMipsDebug()
+{
+	bool bOK = RuntimeVirtualTexture::BuildStreamedMips(RuntimeVirtualTextureComponent, ERuntimeVirtualTextureDebugType::Debug);
+	return bOK ? FReply::Handled() : FReply::Unhandled();
+}
+
 #undef LOCTEXT_NAMESPACE
diff --git a/Engine/Source/Runtime/Engine/Classes/Components/RuntimeVirtualTextureComponent.h b/Engine/Source/Runtime/Engine/Classes/Components/RuntimeVirtualTextureComponent.h
--- a/Engine/Source/Runtime/Engine/Classes/Components/RuntimeVirtualTextureComponent.h
+++ b/Engine/Source/Runtime/Engine/Classes/Components/RuntimeVirtualTextureComponent.h
@@ -19,17 +19,24 @@ private:
 	UPROPERTY(EditAnywhere, DuplicateTransient, Category = VirtualTexture)
 	URuntimeVirtualTexture* VirtualTexture = nullptr;
 
+	/** Use any streaming low mips when rendering in editor. Set true to view and debug the baked streaming low mips. */
+	UPROPERTY(EditAnywhere, AdvancedDisplay, Category = VirtualTexture)
+	bool bUseStreamingLowMipsInEditor = false;
+
 	/** Actor to copy the bounds from to set up the transform. */
 	UPROPERTY(EditAnywhere, DuplicateTransient, Category = TransformFromBounds, meta = (DisplayName = "Source Actor"))
 	AActor* BoundsSourceActor = nullptr;
 
 public:
-	/** Get the runtime virtual texture object set on this component */
+	/** Get the runtime virtual texture object on this component. */
 	URuntimeVirtualTexture* GetVirtualTexture() const { return VirtualTexture; }
 
 	/** Get the runtime virtual texture UV to World transform on this component. */
 	FTransform GetVirtualTextureTransform() const;
 
+	/** Get if we want use any streaming low mips in the runtime virtual texture set on this component. */
+	bool IsStreamingLowMips() const;
+
 #if WITH_EDITOR
 	/** Copy the rotation from BoundsSourceActor to this component. Called by our UI details customization. */
 	void SetRotation();
diff --git a/Engine/Source/Runtime/Engine/Classes/VT/RuntimeVirtualTexture.h b/Engine/Source/Runtime/Engine/Classes/VT/RuntimeVirtualTexture.h
--- a/Engine/Source/Runtime/Engine/Classes/VT/RuntimeVirtualTexture.h
+++ b/Engine/Source/Runtime/Engine/Classes/VT/RuntimeVirtualTexture.h
@@ -37,9 +37,17 @@ protected:
 	int32 TileBorderSize = 2; // 4
 
 	/** Number of low mips to cut from the virtual texture. This can reduce peak virtual texture update cost but will also increase the probability of mip shimmering. */
-	UPROPERTY(EditAnywhere, BlueprintReadOnly, Category = Size, meta = (UIMin = "0", UIMax = "5", DisplayName = "Number of low mips to remove from the virtual texture"))
+	UPROPERTY(EditAnywhere, BlueprintReadOnly, Category = LowMips, meta = (UIMin = "0", UIMax = "5", DisplayName = "Number of low mips to remove from the virtual texture"))
 	int32 RemoveLowMips = 0;
 
+	/** Number of low mips to serialize and stream for the virtual texture. This can reduce rendering update cost. */
+	UPROPERTY(EditAnywhere, BlueprintReadOnly, Category = LowMips, meta = (UIMin = "0", UIMax = "7", DisplayName = "Number of low mips to stream to the virtual texture"))
+	int32 StreamLowMips = 0;
+
+	/** Texture object containing streamed low mips. */
+	UPROPERTY(VisibleAnywhere, Category = LowMips, meta = (DisplayName = "Streaming low mip texture"))
+	class URuntimeVirtualTextureStreamingProxy* StreamingTexture;
+
 	/** Enable usage of the virtual texture. When disabled there is no rendering into the virtual texture, and sampling will return zero values. */
 	UPROPERTY(EditAnywhere, BlueprintReadOnly, Category = Layout, meta = (DisplayName = "Enable virtual texture"))
 	bool bEnable = true;
@@ -57,8 +65,10 @@ public:
 	int32 GetTileSize() const { return 1 << FMath::Clamp(TileSize + 6, 6, 10); }
 	/** Public getter for virtual texture tile border size */
 	int32 GetTileBorderSize() const { return 2 * FMath::Clamp(TileBorderSize, 0, 4); }
-	/** Public getter for virtual texture tile border size */
+	/** Public getter for virtual texture removed low mips */
 	int32 GetRemoveLowMips() const { return RemoveLowMips; }
+	/** Public getter for virtual texture streaming low mips */
+	int32 GetStreamLowMips() const { return StreamLowMips; }
 
 	/** Returns an approximate estimated value for the memory used by the page table texture. */
 	int32 GetEstimatedPageTableTextureMemoryKb() const;
@@ -67,8 +77,13 @@ public:
 
 	/** Get virtual texture description based on the properties of this object and the passed in volume transform. */
 	void GetProducerDescription(FVTProducerDescription& OutDesc, FTransform const& VolumeToWorld) const;
+
+	/** Returns number of texture layers in the virtual texture */
+	int32 GetLayerCount() const;
 	/** Return true if the virtual texture layer should be sampled as sRGB */
 	bool IsLayerSRGB(int32 LayerIndex) const;
+	/** Return true if the virtual texture layer contains an alpha channel */
+	bool IsLayerAlpha(int32 LayerIndex) const;
 
 	/** (Re)Initialize this object. Call this whenever we modify the producer or transform. */
 	void Initialize(IVirtualTexture* InProducer, FTransform const& VolumeToWorld);
@@ -84,6 +99,16 @@ public:
 	/** Getter for the shader uniform parameters. */
 	FVector4 GetUniformParameter(int32 Index);
 
+#if WITH_EDITOR
+	/** Get a hash of the current state to use for streaming texture invalidation. */
+	uint32 GetStreamingTextureBuildHash() const;
+	/** Initialize the low mip streaming texture with the passed in size and data. */
+	void InitializeStreamingTexture(uint32 InSizeX, uint32 InSizeY, uint8* InData);
+#endif
+
+	/** Create the streaming texture producer to wrap an existing producer. */
+	IVirtualTexture* CreateStreamingTextureProducer(IVirtualTexture* InProducer, int32 InMaxLevel, int32& OutTransitionLevel) const;
+
 protected:
 	/** Initialize the render resources. This kicks off render thread work. */
 	void InitResource(IVirtualTexture* InProducer, FTransform const& VolumeToWorld);
diff --git a/Engine/Source/Runtime/Engine/Private/Components/RuntimeVirtualTextureComponent.cpp b/Engine/Source/Runtime/Engine/Private/Components/RuntimeVirtualTextureComponent.cpp
--- a/Engine/Source/Runtime/Engine/Private/Components/RuntimeVirtualTextureComponent.cpp
+++ b/Engine/Source/Runtime/Engine/Private/Components/RuntimeVirtualTextureComponent.cpp
@@ -50,6 +50,15 @@ FTransform URuntimeVirtualTextureComponent::GetVirtualTextureTransform() const
 	return FTransform(FVector(-0.5f, -0.5f, 0.f)) * GetComponentTransform();
 }
 
+bool URuntimeVirtualTextureComponent::IsStreamingLowMips() const
+{
+#if WITH_EDITOR
+	return bUseStreamingLowMipsInEditor;
+#else
+	return true;
+#endif
+}
+
 #if WITH_EDITOR
 
 void URuntimeVirtualTextureComponent::SetRotation()
diff --git a/Engine/Source/Runtime/Engine/Private/VT/RuntimeVirtualTexture.cpp b/Engine/Source/Runtime/Engine/Private/VT/RuntimeVirtualTexture.cpp
--- a/Engine/Source/Runtime/Engine/Private/VT/RuntimeVirtualTexture.cpp
+++ b/Engine/Source/Runtime/Engine/Private/VT/RuntimeVirtualTexture.cpp
@@ -4,8 +4,11 @@
 
 #include "EngineModule.h"
 #include "RendererInterface.h"
+#include "VT/VirtualTextureBuildSettings.h"
 #include "VT/RuntimeVirtualTextureNotify.h"
-
+#include "VT/RuntimeVirtualTextureStreamingProxy.h"
+#include "VT/UploadingVirtualTexture.h"
+#include "VT/VirtualTextureLevelRedirector.h"
 
 namespace
 {
@@ -161,6 +164,18 @@ private:
 };
 
 
+URuntimeVirtualTextureStreamingProxy::URuntimeVirtualTextureStreamingProxy(const FObjectInitializer& ObjectInitializer)
+	: UTexture2D(ObjectInitializer)
+	, BuildHash(0)
+{
+}
+
+void URuntimeVirtualTextureStreamingProxy::GetVirtualTextureBuildSettings(FVirtualTextureBuildSettings& OutSettings) const
+{
+	OutSettings = Settings;
+}
+
+
 URuntimeVirtualTexture::URuntimeVirtualTexture(const FObjectInitializer& ObjectInitializer)
 	: Super(ObjectInitializer)
 {
@@ -233,6 +248,24 @@ void URuntimeVirtualTexture::GetProducerDescription(FVTProducerDescription& OutD
 	}
 }
 
+int32 URuntimeVirtualTexture::GetLayerCount() const
+{
+	switch (MaterialType)
+	{
+	case ERuntimeVirtualTextureMaterialType::BaseColor:
+	case ERuntimeVirtualTextureMaterialType::WorldHeight:
+		return 1;
+	case ERuntimeVirtualTextureMaterialType::BaseColor_Normal:
+	case ERuntimeVirtualTextureMaterialType::BaseColor_Normal_Specular:
+		return 2;
+	default:
+		// Implement logic for any missing material types
+		check(false);
+	}
+
+	return 1;
+}
+
 bool URuntimeVirtualTexture::IsLayerSRGB(int32 LayerIndex) const
 {
 	switch (MaterialType)
@@ -252,6 +285,24 @@ bool URuntimeVirtualTexture::IsLayerSRGB(int32 LayerIndex) const
 	return false;
 }
 
+bool URuntimeVirtualTexture::IsLayerAlpha(int32 LayerIndex) const
+{
+	switch (MaterialType)
+	{
+	case ERuntimeVirtualTextureMaterialType::BaseColor:
+	case ERuntimeVirtualTextureMaterialType::BaseColor_Normal:
+	case ERuntimeVirtualTextureMaterialType::WorldHeight:
+		return false;
+	case ERuntimeVirtualTextureMaterialType::BaseColor_Normal_Specular:
+		return LayerIndex == 1;
+	default:
+		// Implement logic for any missing material types
+		check(false);
+	}
+
+	return false;
+}
+
 int32 URuntimeVirtualTexture::GetEstimatedPageTableTextureMemoryKb() const
 {
 	//todo[vt]: Estimate memory usage
@@ -326,8 +377,92 @@ void URuntimeVirtualTexture::PostEditChangeProperty(FPropertyChangedEvent& Prope
 {
 	Super::PostEditChangeProperty(PropertyChangedEvent);
 
+	// Invalidate StreamingTexture if it is no longer compatible with this object
+	if (StreamingTexture != nullptr && StreamingTexture->BuildHash != GetStreamingTextureBuildHash())
+	{
+		StreamingTexture = nullptr;
+	}
+
 	RuntimeVirtualTexture::NotifyComponents(this);
 	RuntimeVirtualTexture::NotifyPrimitives(this);
 }
 
+uint32 URuntimeVirtualTexture::GetStreamingTextureBuildHash() const
+{
+	union FPackedSettings
+	{
+		uint32 PackedValue;
+		struct
+		{
+			uint32 MaterialType : 3;
+			uint32 CompressTextures : 1;
+			uint32 Size : 4;
+			uint32 TileSize : 4;
+			uint32 TileBorderSize : 4;
+			uint32 RemoveLowMips : 4;
+			uint32 StreamLowMips : 4;
+		};
+	};
+
+	FPackedSettings Settings;
+	Settings.PackedValue = 0;
+	Settings.MaterialType = (uint32)MaterialType;
+	Settings.CompressTextures = (uint32)bCompressTextures;
+	Settings.Size = (uint32)Size;
+	Settings.TileSize = (uint32)TileSize;
+	Settings.RemoveLowMips = (uint32)RemoveLowMips;
+	Settings.StreamLowMips = (uint32)StreamLowMips;
+
+	return Settings.PackedValue;
+}
+
+void URuntimeVirtualTexture::InitializeStreamingTexture(uint32 InSizeX, uint32 InSizeY, uint8* InData)
+{
+	// Release current producer. 
+	// It may reference data inside StreamingTexture which could be garbage collected any time from now.
+	InitNullResource();
+
+	StreamingTexture = NewObject<URuntimeVirtualTextureStreamingProxy>(GetOutermost(), TEXT("StreamingTexture"));
+	StreamingTexture->VirtualTextureStreaming = true;
+	StreamingTexture->MipGenSettings = TMGS_SimpleAverage;
+
+	StreamingTexture->Settings.Init();
+	StreamingTexture->Settings.TileSize = GetTileSize();
+	StreamingTexture->Settings.TileBorderSize = GetTileBorderSize();
+
+	StreamingTexture->BuildHash = GetStreamingTextureBuildHash();
+
+	const int32 LayerCount = GetLayerCount();
+	for (int32 Layer = 0; Layer < LayerCount; Layer++)
+	{
+		FTextureFormatSettings FormatSettings;
+		FormatSettings.SRGB = IsLayerSRGB(Layer);
+		FormatSettings.CompressionNone = !bCompressTextures;
+		FormatSettings.CompressionNoAlpha = !IsLayerAlpha(Layer);
+		StreamingTexture->SetLayerFormatSettings(Layer, FormatSettings);
+	}
+
+	const ETextureSourceFormat LayerFormats[] = { TSF_BGRA8, TSF_BGRA8 };
+	StreamingTexture->Source.InitLayered(InSizeX, InSizeY, 1, LayerCount, 1, LayerFormats, InData);
+
+	StreamingTexture->PostEditChange();
+}
+
 #endif
+
+IVirtualTexture* URuntimeVirtualTexture::CreateStreamingTextureProducer(IVirtualTexture* InProducer, int32 InMaxLevel, int32& OutTransitionLevel) const
+{
+	if (StreamingTexture != nullptr)
+	{
+		FTexturePlatformData** StreamingTextureData = StreamingTexture->GetRunningPlatformData();
+		if (StreamingTextureData != nullptr && *StreamingTextureData != nullptr)
+		{
+			OutTransitionLevel = InMaxLevel - StreamLowMips;
+			IVirtualTexture* StreamingProducer = new FUploadingVirtualTexture((*StreamingTextureData)->VTData, 0);
+			return new FVirtualTextureLevelRedirector(InProducer, StreamingProducer, OutTransitionLevel);
+		}
+	}
+	// Can't create a streaming producer so return original producer.
+	OutTransitionLevel = InMaxLevel;
+	return InProducer;
+}
diff --git a/Engine/Source/Runtime/Engine/Private/VT/RuntimeVirtualTextureStreamingProxy.h b/Engine/Source/Runtime/Engine/Private/VT/RuntimeVirtualTextureStreamingProxy.h
new file mode 100644
--- /dev/null
+++ b/Engine/Source/Runtime/Engine/Private/VT/RuntimeVirtualTextureStreamingProxy.h
@@ -0,0 +1,27 @@
+// Copyright 1998-2019 Epic Games, Inc. All Rights Reserved.
+
+#pragma once
+
+#include "CoreMinimal.h"
+#include "Engine/Texture.h"
+#include "VT/VirtualTextureBuildSettings.h"
+#include "RuntimeVirtualTextureStreamingProxy.generated.h"
+
+/** Streaming virtual texture used to store the low mips in a URuntimeVirtualTexture. */
+UCLASS(ClassGroup = Rendering)
+class URuntimeVirtualTextureStreamingProxy : public UTexture2D
+{
+	GENERATED_UCLASS_BODY()
+
+	/** Virtual texture build settings. These should match those of any owning URuntimeVirtualTexture. */
+	UPROPERTY()
+	FVirtualTextureBuildSettings Settings;
+
+	/** Hash of settings used when building this texture. Used to invalidate when build settings have changed. */
+	UPROPERTY()
+	uint32 BuildHash;
+
+	//~ Begin UTexture Interface.
+	virtual void GetVirtualTextureBuildSettings(FVirtualTextureBuildSettings& OutSettings) const override;
+	//~ End UTexture Interface.
+};
diff --git a/Engine/Source/Runtime/Engine/Private/VT/VirtualTextureLevelRedirector.cpp b/Engine/Source/Runtime/Engine/Private/VT/VirtualTextureLevelRedirector.cpp
new file mode 100644
--- /dev/null
+++ b/Engine/Source/Runtime/Engine/Private/VT/VirtualTextureLevelRedirector.cpp
@@ -0,0 +1,43 @@
+// Copyright 1998-2019 Epic Games, Inc. All Rights Reserved.
+
+#include "VT/VirtualTextureLevelRedirector.h"
+
+FVirtualTextureLevelRedirector::FVirtualTextureLevelRedirector(IVirtualTexture* InVirtualTexture0, IVirtualTexture* InVirtualTexture1, int32 InTransitionLevel)
+	: VirtualTextures{ InVirtualTexture0, InVirtualTexture1 }
+	, TransitionLevel(InTransitionLevel)
+{
+}
+
+FVirtualTextureLevelRedirector::~FVirtualTextureLevelRedirector()
+{
+	delete VirtualTextures[0];
+	delete VirtualTextures[1];
+}
+
+FVTRequestPageResult FVirtualTextureLevelRedirector::RequestPageData(
+	const FVirtualTextureProducerHandle& ProducerHandle,
+	uint8 LayerMask,
+	uint8 vLevel,
+	uint32 vAddress,
+	EVTRequestPagePriority Priority)
+{
+	int32 VirtualTextureIndex = vLevel <= TransitionLevel ? 0 : 1;
+	int32 vLevelOffset = vLevel <= TransitionLevel ? 0 : TransitionLevel + 1;
+	return VirtualTextures[VirtualTextureIndex]->RequestPageData(ProducerHandle, LayerMask, vLevel - vLevelOffset, vAddress, Priority);
+}
+
+IVirtualTextureFinalizer* FVirtualTextureLevelRedirector::ProducePageData(
+	FRHICommandListImmediate& RHICmdList,
+	ERHIFeatureLevel::Type FeatureLevel,
+	EVTProducePageFlags Flags,
+	const FVirtualTextureProducerHandle& ProducerHandle,
+	uint8 LayerMask,
+	uint8 vLevel,
+	uint32 vAddress,
+	uint64 RequestHandle,
+	const FVTProduceTargetLayer* TargetLayers)
+{
+	int32 VirtualTextureIndex = vLevel <= TransitionLevel ? 0 : 1;
+	int32 vLevelOffset = vLevel <= TransitionLevel ? 0 : TransitionLevel + 1;
+	return VirtualTextures[VirtualTextureIndex]->ProducePageData(RHICmdList, FeatureLevel, Flags, ProducerHandle, LayerMask, vLevel - vLevelOffset, vAddress, RequestHandle, TargetLayers);
+}
diff --git a/Engine/Source/Runtime/Engine/Private/VT/VirtualTextureLevelRedirector.h b/Engine/Source/Runtime/Engine/Private/VT/VirtualTextureLevelRedirector.h
new file mode 100644
--- /dev/null
+++ b/Engine/Source/Runtime/Engine/Private/VT/VirtualTextureLevelRedirector.h
@@ -0,0 +1,40 @@
+// Copyright 1998-2019 Epic Games, Inc. All Rights Reserved.
+
+#pragma once
+
+#include "CoreMinimal.h"
+#include "VirtualTexturing.h"
+
+/** IVirtualTexture implementation that redirects requests to one of two children depending on vLevel. */
+class FVirtualTextureLevelRedirector : public IVirtualTexture
+{
+public:
+	FVirtualTextureLevelRedirector(IVirtualTexture* InVirtualTexture0, IVirtualTexture* InVirtualTexture1, int32 InTransitionLevel);
+	virtual ~FVirtualTextureLevelRedirector();
+
+	//~ Begin IVirtualTexture Interface.
+	virtual FVTRequestPageResult RequestPageData(
+		const FVirtualTextureProducerHandle& ProducerHandle,
+		uint8 LayerMask,
+		uint8 vLevel,
+		uint32 vAddress,
+		EVTRequestPagePriority Priority
+	) override;
+
+	virtual IVirtualTextureFinalizer* ProducePageData(
+		FRHICommandListImmediate& RHICmdList,
+		ERHIFeatureLevel::Type FeatureLevel,
+		EVTProducePageFlags Flags,
+		const FVirtualTextureProducerHandle& ProducerHandle,
+		uint8 LayerMask,
+		uint8 vLevel,
+		uint32 vAddress,
+		uint64 RequestHandle,
+		const FVTProduceTargetLayer* TargetLayers
+	) override;
+	//~ End IVirtualTexture Interface.
+
+private:
+	IVirtualTexture* VirtualTextures[2];
+	int32 TransitionLevel;
+};
diff --git a/Engine/Source/Runtime/Renderer/Private/VT/RuntimeVirtualTextureSceneProxy.cpp b/Engine/Source/Runtime/Renderer/Private/VT/RuntimeVirtualTextureSceneProxy.cpp
--- a/Engine/Source/Runtime/Renderer/Private/VT/RuntimeVirtualTextureSceneProxy.cpp
+++ b/Engine/Source/Runtime/Renderer/Private/VT/RuntimeVirtualTextureSceneProxy.cpp
@@ -17,16 +17,12 @@ FRuntimeVirtualTextureSceneProxy::FRuntimeVirtualTextureSceneProxy(URuntimeVirtu
 {
 	if (InComponent->GetVirtualTexture() != nullptr && InComponent->GetVirtualTexture()->GetEnabled())
 	{
-		VirtualTexture = InComponent->GetVirtualTexture();
-
 		// We store a ProducerId here so that we will be able to find our SceneIndex from the Producer during rendering.
 		// We will need the SceneIndex to determine which primitives should render to this Producer.
 		ProducerId = ProducerIdGenerator++;
 
-		const ERuntimeVirtualTextureMaterialType MaterialType = VirtualTexture->GetMaterialType();
-
-		// Transform is based on bottom left of the URuntimeVirtualTextureComponent unit box (which is centered on the origin)
-		Transform = FTransform(FVector(-0.5f, -0.5f, 0.f)) * InComponent->GetComponentTransform();
+		VirtualTexture = InComponent->GetVirtualTexture();
+		Transform = InComponent->GetVirtualTextureTransform();
 
 		// The producer description is calculated using the transform to determine the aspect ratio
 		FVTProducerDescription Desc;
@@ -34,9 +30,21 @@ FRuntimeVirtualTextureSceneProxy::FRuntimeVirtualTextureSceneProxy(URuntimeVirtu
 		VirtualTextureSize = FIntPoint(Desc.WidthInBlocks * Desc.BlockWidthInTiles * Desc.TileSize, Desc.HeightInBlocks * Desc.BlockHeightInTiles * Desc.TileSize);
 		MaxDirtyLevel = Desc.MaxLevel;
 
+		const ERuntimeVirtualTextureMaterialType MaterialType = VirtualTexture->GetMaterialType();
+
 		// The Producer object created here will be passed into the virtual texture system which will take ownership.
+		IVirtualTexture* Producer = new FRuntimeVirtualTextureProducer(Desc, ProducerId, MaterialType, InComponent->GetScene(), Transform);
+
+		if (InComponent->IsStreamingLowMips() && VirtualTexture->GetStreamLowMips() > 0)
+		{
+			// Wrap our producer to use a streaming producer for low mips
+			int32 StreamingTransitionLevel;
+			Producer = VirtualTexture->CreateStreamingTextureProducer(Producer, Desc.MaxLevel, StreamingTransitionLevel);
+			// Any dirty flushes don't need to flush the streaming mips (they only change with a build step).
+			MaxDirtyLevel = StreamingTransitionLevel;
+		}
+
 		// The Initialize() call will allocate the virtual texture by spawning work on the render thread.
-		FRuntimeVirtualTextureProducer* Producer = new FRuntimeVirtualTextureProducer(Desc, ProducerId, MaterialType, InComponent->GetScene(), Transform);
 		VirtualTexture->Initialize(Producer, Transform);
 	}
 }
