Some non-unity build fixes

diff --git a/Engine/Source/Editor/VirtualTexturingEditor/Private/RuntimeVirtualTextureDetailsCustomization.cpp b/Engine/Source/Editor/VirtualTexturingEditor/Private/RuntimeVirtualTextureDetailsCustomization.cpp
--- a/Engine/Source/Editor/VirtualTexturingEditor/Private/RuntimeVirtualTextureDetailsCustomization.cpp
+++ b/Engine/Source/Editor/VirtualTexturingEditor/Private/RuntimeVirtualTextureDetailsCustomization.cpp
@@ -8,6 +8,7 @@
 #include "DetailLayoutBuilder.h"
 #include "DetailWidgetRow.h"
 #include "RuntimeVirtualTextureBuild.h"
+#include "ScopedTransaction.h"
 #include "SResetToDefaultMenu.h"
 #include "VT/RuntimeVirtualTexture.h"
 #include "Widgets/Input/SButton.h"
diff --git a/Engine/Source/Runtime/Engine/Private/VT/RuntimeVirtualTextureStreamingProxy.h b/Engine/Source/Runtime/Engine/Private/VT/RuntimeVirtualTextureStreamingProxy.h
--- a/Engine/Source/Runtime/Engine/Private/VT/RuntimeVirtualTextureStreamingProxy.h
+++ b/Engine/Source/Runtime/Engine/Private/VT/RuntimeVirtualTextureStreamingProxy.h
@@ -3,7 +3,7 @@
 #pragma once
 
 #include "CoreMinimal.h"
-#include "Engine/Texture.h"
+#include "Engine/Texture2D.h"
 #include "VT/VirtualTextureBuildSettings.h"
 #include "RuntimeVirtualTextureStreamingProxy.generated.h"
 
diff --git a/Engine/Source/Runtime/Engine/Private/VT/VirtualTextureBuildSettings.cpp b/Engine/Source/Runtime/Engine/Private/VT/VirtualTextureBuildSettings.cpp
--- a/Engine/Source/Runtime/Engine/Private/VT/VirtualTextureBuildSettings.cpp
+++ b/Engine/Source/Runtime/Engine/Private/VT/VirtualTextureBuildSettings.cpp
@@ -2,6 +2,8 @@
 
 #include "VT/VirtualTextureBuildSettings.h"
 
+#include "HAL/IConsoleManager.h"
+
 static TAutoConsoleVariable<int32> CVarVTCompressCrunch(
 	TEXT("r.VT.EnableCompressCrunch"),
 	0,
