Some more non-unity build fixes

diff --git a/Engine/Source/Editor/VirtualTexturingEditor/Private/RuntimeVirtualTextureBuild.cpp b/Engine/Source/Editor/VirtualTexturingEditor/Private/RuntimeVirtualTextureBuild.cpp
--- a/Engine/Source/Editor/VirtualTexturingEditor/Private/RuntimeVirtualTextureBuild.cpp
+++ b/Engine/Source/Editor/VirtualTexturingEditor/Private/RuntimeVirtualTextureBuild.cpp
@@ -5,12 +5,13 @@
 #include "Components/RuntimeVirtualTextureComponent.h"
 #include "Engine/TextureRenderTarget2D.h"
 #include "Misc/ScopedSlowTask.h"
+#include "SceneInterface.h"
 #include "VT/RuntimeVirtualTexture.h"
 #include "VT/RuntimeVirtualTextureRender.h"
 
 namespace RuntimeVirtualTexture
 {
-	bool HasStreamedMips(class URuntimeVirtualTextureComponent* InComponent)
+	bool HasStreamedMips(URuntimeVirtualTextureComponent* InComponent)
 	{
 		if (InComponent == nullptr)
 		{
@@ -31,7 +32,7 @@ namespace RuntimeVirtualTexture
 		return true;
 	}
 
-	bool BuildStreamedMips(class URuntimeVirtualTextureComponent* InComponent, ERuntimeVirtualTextureDebugType DebugType)
+	bool BuildStreamedMips(URuntimeVirtualTextureComponent* InComponent, ERuntimeVirtualTextureDebugType DebugType)
 	{
 		if (!HasStreamedMips(InComponent))
 		{
diff --git a/Engine/Source/Editor/VirtualTexturingEditor/Private/RuntimeVirtualTextureBuild.h b/Engine/Source/Editor/VirtualTexturingEditor/Private/RuntimeVirtualTextureBuild.h
--- a/Engine/Source/Editor/VirtualTexturingEditor/Private/RuntimeVirtualTextureBuild.h
+++ b/Engine/Source/Editor/VirtualTexturingEditor/Private/RuntimeVirtualTextureBuild.h
@@ -4,11 +4,14 @@
 
 #include "CoreMinimal.h"
 
+enum class ERuntimeVirtualTextureDebugType;
+class URuntimeVirtualTextureComponent;
+
 namespace RuntimeVirtualTexture
 {
 	/** Returns true if the component describes a runtime virtual texture that has streaming mips. */
-	bool HasStreamedMips(class URuntimeVirtualTextureComponent* InComponent);
+	bool HasStreamedMips(URuntimeVirtualTextureComponent* InComponent);
 
 	/** Build the streaming mips and store in the component's associated URuntimeVirtualTexture object. */
-	bool BuildStreamedMips(class URuntimeVirtualTextureComponent* InComponent, enum class ERuntimeVirtualTextureDebugType DebugType);
+	bool BuildStreamedMips(URuntimeVirtualTextureComponent* InComponent, ERuntimeVirtualTextureDebugType DebugType);
 };
diff --git a/Engine/Source/Editor/VirtualTexturingEditor/Private/RuntimeVirtualTextureThumbnailRenderer.cpp b/Engine/Source/Editor/VirtualTexturingEditor/Private/RuntimeVirtualTextureThumbnailRenderer.cpp
--- a/Engine/Source/Editor/VirtualTexturingEditor/Private/RuntimeVirtualTextureThumbnailRenderer.cpp
+++ b/Engine/Source/Editor/VirtualTexturingEditor/Private/RuntimeVirtualTextureThumbnailRenderer.cpp
@@ -3,8 +3,12 @@
 #include "RuntimeVirtualTextureThumbnailRenderer.h"
 
 #include "Components/RuntimeVirtualTextureComponent.h"
+#include "RenderingThread.h"
+#include "SceneInterface.h"
 #include "VT/RuntimeVirtualTexture.h"
 #include "VT/RuntimeVirtualTextureRender.h"
+#include "UnrealClient.h"
+#include "UObject/UObjectIterator.h"
 
 namespace
 {
