Add missing pragma once

diff --git a/Engine/Source/Editor/VirtualTexturingEditor/Public/VirtualTexturingEditorModule.h b/Engine/Source/Editor/VirtualTexturingEditor/Public/VirtualTexturingEditorModule.h
--- a/Engine/Source/Editor/VirtualTexturingEditor/Public/VirtualTexturingEditorModule.h
+++ b/Engine/Source/Editor/VirtualTexturingEditor/Public/VirtualTexturingEditorModule.h
@@ -1,5 +1,7 @@
 // Copyright 1998-2019 Epic Games, Inc. All Rights Reserved.
 
+#pragma once
+
 #include "CoreMinimal.h"
 #include "Modules/ModuleInterface.h"
 
