Test that scene is ready to render and that material proxies are up to date before drawing RVT thumbnail
Fixes check firing in RHIUpdateUniformBuffer()

diff --git a/Engine/Source/Editor/VirtualTexturingEditor/Private/RuntimeVirtualTextureThumbnailRenderer.cpp b/Engine/Source/Editor/VirtualTexturingEditor/Private/RuntimeVirtualTextureThumbnailRenderer.cpp
--- a/Engine/Source/Editor/VirtualTexturingEditor/Private/RuntimeVirtualTextureThumbnailRenderer.cpp
+++ b/Engine/Source/Editor/VirtualTexturingEditor/Private/RuntimeVirtualTextureThumbnailRenderer.cpp
@@ -3,6 +3,7 @@
 #include "RuntimeVirtualTextureThumbnailRenderer.h"
 
 #include "Components/RuntimeVirtualTextureComponent.h"
+#include "MaterialShared.h"
 #include "RenderingThread.h"
 #include "SceneInterface.h"
 #include "VT/RuntimeVirtualTexture.h"
@@ -43,8 +44,14 @@ bool URuntimeVirtualTextureThumbnailRenderer::CanVisualizeAsset(UObject* Object)
 		{
 			// We need a matching URuntimeVirtualTextureComponent in a Scene to be able to render a thumbnail
 			URuntimeVirtualTextureComponent* RuntimeVirtualTextureComponent = FindComponent(RuntimeVirtualTexture);
-			FSceneInterface* Scene = RuntimeVirtualTextureComponent != nullptr ? RuntimeVirtualTextureComponent->GetScene() : nullptr;
-			return Scene != nullptr;
+			if (RuntimeVirtualTextureComponent != nullptr)
+			{
+				FSceneInterface* Scene = RuntimeVirtualTextureComponent->GetScene();
+				if (Scene != nullptr && RuntimeVirtualTexture::IsSceneReadyToRender(Scene->GetRenderScene()))
+				{
+					return true;
+				}
+			}
 		}
 	}
 
@@ -70,6 +77,8 @@ void URuntimeVirtualTextureThumbnailRenderer::Draw(UObject* Object, int32 X, int
 	ENQUEUE_RENDER_COMMAND(BakeStreamingTextureTileCommand)(
 		[Scene, VirtualTextureSceneIndex, MaterialType, RenderTarget, DestRect, Transform, MaxLevel](FRHICommandListImmediate& RHICmdList)
 	{
+		FMaterialRenderProxy::UpdateDeferredCachedUniformExpressions();
+
 		RuntimeVirtualTexture::RenderPage(
 			RHICmdList,
 			Scene->GetRenderScene(),
diff --git a/Engine/Source/Runtime/Renderer/Private/VT/RuntimeVirtualTextureProducer.cpp b/Engine/Source/Runtime/Renderer/Private/VT/RuntimeVirtualTextureProducer.cpp
--- a/Engine/Source/Runtime/Renderer/Private/VT/RuntimeVirtualTextureProducer.cpp
+++ b/Engine/Source/Runtime/Renderer/Private/VT/RuntimeVirtualTextureProducer.cpp
@@ -20,11 +20,7 @@ FRuntimeVirtualTextureFinalizer::FRuntimeVirtualTextureFinalizer(FVTProducerDesc
 
 bool FRuntimeVirtualTextureFinalizer::IsReady()
 {
-	//todo[vt]: 
-	// Test if we have everything we need to render (mips loaded etc).
-
-	// Test scene is loaded and has been updated once by main rendering passes
-	return Scene != nullptr && Scene->GetRenderScene() != nullptr && Scene->GetRenderScene()->GetFrameNumber() > 1;
+	return RuntimeVirtualTexture::IsSceneReadyToRender(Scene->GetRenderScene());
 }
 
 void FRuntimeVirtualTextureFinalizer::InitProducer(const FVirtualTextureProducerHandle& ProducerHandle)
diff --git a/Engine/Source/Runtime/Renderer/Private/VT/RuntimeVirtualTextureRender.cpp b/Engine/Source/Runtime/Renderer/Private/VT/RuntimeVirtualTextureRender.cpp
--- a/Engine/Source/Runtime/Renderer/Private/VT/RuntimeVirtualTextureRender.cpp
+++ b/Engine/Source/Runtime/Renderer/Private/VT/RuntimeVirtualTextureRender.cpp
@@ -640,6 +640,14 @@ namespace RuntimeVirtualTexture
 	};
 
 
+	bool IsSceneReadyToRender(FScene* Scene)
+	{
+		// Test scene is loaded and has been updated once by main rendering passes
+		// This function gets called on the main thread, so accessing scene frame number is not strictly thread safe, but we can probably
+		// assume that frame number is always increasing, and so the test is conservative
+		return Scene != nullptr && Scene->GetRenderScene() != nullptr && Scene->GetRenderScene()->GetFrameNumber() > 1;
+	}
+
 	void RenderPage(
 		FRHICommandListImmediate& RHICmdList,
 		FScene* Scene,
diff --git a/Engine/Source/Runtime/Renderer/Public/VT/RuntimeVirtualTextureRender.h b/Engine/Source/Runtime/Renderer/Public/VT/RuntimeVirtualTextureRender.h
--- a/Engine/Source/Runtime/Renderer/Public/VT/RuntimeVirtualTextureRender.h
+++ b/Engine/Source/Runtime/Renderer/Public/VT/RuntimeVirtualTextureRender.h
@@ -22,6 +22,9 @@ namespace RuntimeVirtualTexture
 	RENDERER_API uint32 GetRuntimeVirtualTextureSceneIndex_GameThread(URuntimeVirtualTextureComponent* InComponent);
 #endif
 
+	/** Returns true if the FScene is initialized for rendering to runtime virtual textures. */
+	RENDERER_API bool IsSceneReadyToRender(FScene* Scene);
+
 	/**
 	 * Render a single page of a virtual texture with a given material.
 	 * todo[vt]: Likely to be more optimal to batch several pages at a time and share setup/visibility/render targets.
