Change RVT size to be defined by TileCount instead of directly by size
This makes it easier to limit according to maximum page table size
Supports backwards compatibility with old data

diff --git a/Engine/Source/Editor/VirtualTexturingEditor/Classes/RuntimeVirtualTextureDetailsCustomization.h b/Engine/Source/Editor/VirtualTexturingEditor/Classes/RuntimeVirtualTextureDetailsCustomization.h
--- a/Engine/Source/Editor/VirtualTexturingEditor/Classes/RuntimeVirtualTextureDetailsCustomization.h
+++ b/Engine/Source/Editor/VirtualTexturingEditor/Classes/RuntimeVirtualTextureDetailsCustomization.h
@@ -26,10 +26,11 @@ protected:
 private:
 	class URuntimeVirtualTexture* VirtualTexture;
 
-	TSharedPtr<class STextBlock> SizeText;
+	TSharedPtr<class STextBlock> TileCountText;
 	TSharedPtr<class STextBlock> TileSizeText;
 	TSharedPtr<class STextBlock> TileBorderSizeText;
 
+	TSharedPtr<class STextBlock> SizeText;
 	TSharedPtr<class STextBlock> PageTableTextureMemoryText;
 	TSharedPtr<class STextBlock> PhysicalTextureMemoryText;
 };
diff --git a/Engine/Source/Editor/VirtualTexturingEditor/Private/RuntimeVirtualTextureDetailsCustomization.cpp b/Engine/Source/Editor/VirtualTexturingEditor/Private/RuntimeVirtualTextureDetailsCustomization.cpp
--- a/Engine/Source/Editor/VirtualTexturingEditor/Private/RuntimeVirtualTextureDetailsCustomization.cpp
+++ b/Engine/Source/Editor/VirtualTexturingEditor/Private/RuntimeVirtualTextureDetailsCustomization.cpp
@@ -95,7 +95,7 @@ void FRuntimeVirtualTextureDetailsCustomization::CustomizeDetails(IDetailLayoutB
 
 	// Add size helpers
 	IDetailCategoryBuilder& SizeCategory = DetailBuilder.EditCategory("Size", FText::GetEmpty());
-	AddTextToProperty(DetailBuilder, SizeCategory, "Size", SizeText);
+	AddTextToProperty(DetailBuilder, SizeCategory, "TileCount", TileCountText);
 	AddTextToProperty(DetailBuilder, SizeCategory, "TileSize", TileSizeText);
 	AddTextToProperty(DetailBuilder, SizeCategory, "TileBorderSize", TileBorderSizeText);
 
@@ -107,6 +107,14 @@ void FRuntimeVirtualTextureDetailsCustomization::CustomizeDetails(IDetailLayoutB
 	[
 		SNew(SVerticalBox)
 
+		+ SVerticalBox::Slot()
+		.AutoHeight()
+		.VAlign(VAlign_Center)
+		.Padding(4.0f)
+		[
+			SAssignNew(SizeText, STextBlock)
+		]
+
 		+ SVerticalBox::Slot()
 		.AutoHeight()
 		.VAlign(VAlign_Center)
@@ -125,7 +133,7 @@ void FRuntimeVirtualTextureDetailsCustomization::CustomizeDetails(IDetailLayoutB
 	];
 
 	// Add refresh callback for all properties 
-	DetailBuilder.GetProperty(FName(TEXT("Size")))->SetOnPropertyValueChanged(FSimpleDelegate::CreateSP(this, &FRuntimeVirtualTextureDetailsCustomization::RefreshDetails));
+	DetailBuilder.GetProperty(FName(TEXT("TileCount")))->SetOnPropertyValueChanged(FSimpleDelegate::CreateSP(this, &FRuntimeVirtualTextureDetailsCustomization::RefreshDetails));
 	DetailBuilder.GetProperty(FName(TEXT("TileSize")))->SetOnPropertyValueChanged(FSimpleDelegate::CreateSP(this, &FRuntimeVirtualTextureDetailsCustomization::RefreshDetails));
 	DetailBuilder.GetProperty(FName(TEXT("TileBorderSize")))->SetOnPropertyValueChanged(FSimpleDelegate::CreateSP(this, &FRuntimeVirtualTextureDetailsCustomization::RefreshDetails));
 	DetailBuilder.GetProperty(FName(TEXT("MaterialType")))->SetOnPropertyValueChanged(FSimpleDelegate::CreateSP(this, &FRuntimeVirtualTextureDetailsCustomization::RefreshDetails));
@@ -143,10 +151,11 @@ void FRuntimeVirtualTextureDetailsCustomization::RefreshDetails()
 	SizeOptions.UseGrouping = false;
 	SizeOptions.MaximumFractionalDigits = 0;
 
- 	SizeText->SetText(FText::Format(LOCTEXT("Details_Number", "{0}"), FText::AsNumber(VirtualTexture->GetSize(), &SizeOptions)));
+ 	TileCountText->SetText(FText::Format(LOCTEXT("Details_Number", "{0}"), FText::AsNumber(VirtualTexture->GetTileCount(), &SizeOptions)));
  	TileSizeText->SetText(FText::Format(LOCTEXT("Details_Number", "{0}"), FText::AsNumber(VirtualTexture->GetTileSize(), &SizeOptions)));
  	TileBorderSizeText->SetText(FText::Format(LOCTEXT("Details_Number", "{0}"), FText::AsNumber(VirtualTexture->GetTileBorderSize(), &SizeOptions)));
 
+	SizeText->SetText(FText::Format(LOCTEXT("Details_Size", "Virtual Texture Size: {0}"), FText::AsNumber(VirtualTexture->GetSize(), &SizeOptions)));
 	PageTableTextureMemoryText->SetText(FText::Format(LOCTEXT("Details_PageTableMemory", "Page Table Texture Memory (estimated): {0} KiB"), FText::AsNumber(VirtualTexture->GetEstimatedPageTableTextureMemoryKb(), &SizeOptions)));
 	PhysicalTextureMemoryText->SetText(FText::Format(LOCTEXT("Details_PhysicalMemory", "Physical Texture Memory (estimated): {0} KiB"), FText::AsNumber(VirtualTexture->GetEstimatedPhysicalTextureMemoryKb(), &SizeOptions)));
 }
diff --git a/Engine/Source/Runtime/Engine/Classes/VT/RuntimeVirtualTexture.h b/Engine/Source/Runtime/Engine/Classes/VT/RuntimeVirtualTexture.h
--- a/Engine/Source/Runtime/Engine/Classes/VT/RuntimeVirtualTexture.h
+++ b/Engine/Source/Runtime/Engine/Classes/VT/RuntimeVirtualTexture.h
@@ -41,14 +41,22 @@ protected:
 	bool bPrivateSpace = true;
 
 	/** Size of virtual texture along the largest axis. (Actual values increase in powers of 2) */
-	UPROPERTY(EditAnywhere, BlueprintReadOnly, Category = Size, meta = (UIMin = "0", UIMax = "8", DisplayName = "Size of the virtual texture"))
-	int32 Size = 6; // 65536
+	UPROPERTY()
+	int32 Size_DEPRECATED = -1;
+
+	/** 
+	 * Size of virtual texture in tiles. (Actual values increase in powers of 2).
+	 * This replaces the deprecated Size property.
+	 * This is applied to the largest axis in world space and the size for any shorter axis is chosen to maintain aspect ratio.  
+	 */
+	UPROPERTY(EditAnywhere, BlueprintReadOnly, Category = Size, meta = (UIMin = "0", UIMax = "12", DisplayName = "Size of the virtual texture in tiles"))
+	int32 TileCount = 8; // 256
 
 	/** Page tile size. (Actual values increase in powers of 2) */
 	UPROPERTY(EditAnywhere, BlueprintReadOnly, Category = Size, meta = (UIMin = "0", UIMax = "4", DisplayName = "Size of each virtual texture tile"))
 	int32 TileSize = 2; // 256
 
-	/** Page tile border size divided by 2 (Actual values increase in multiples of 2). */
+	/** Page tile border size divided by 2 (Actual values increase in multiples of 2). Higher values trigger a higher anisotropic sampling level. */
 	UPROPERTY(EditAnywhere, BlueprintReadOnly, Category = Size, meta = (UIMin = "0", UIMax = "4", DisplayName = "Border padding for each virtual texture tile"))
 	int32 TileBorderSize = 2; // 4
 
@@ -72,7 +80,9 @@ public:
 	ERuntimeVirtualTextureMaterialType GetMaterialType() const { return MaterialType; }
 
 	/** Public getter for virtual texture size */
-	int32 GetSize() const { return 1 << FMath::Clamp(Size + 10, 10, 18); }
+	int32 GetSize() const { return GetTileCount() * GetTileSize(); }
+	/** Public getter for virtual texture tile count */
+	int32 GetTileCount() const { return 1 << FMath::Clamp(TileCount, 0, 12); }
 	/** Public getter for virtual texture tile size */
 	int32 GetTileSize() const { return 1 << FMath::Clamp(TileSize + 6, 6, 10); }
 	/** Public getter for virtual texture tile border size */
@@ -136,6 +146,7 @@ protected:
 	//~ Begin UObject Interface.
 	virtual void GetAssetRegistryTags(TArray<FAssetRegistryTag>& OutTags) const override;
 	virtual void Serialize(FArchive& Ar) override;
+	virtual void PostLoad() override;
 #if WITH_EDITOR
 	virtual void PostEditChangeProperty(FPropertyChangedEvent& PropertyChangedEvent) override;
 #endif
diff --git a/Engine/Source/Runtime/Engine/Private/VT/RuntimeVirtualTexture.cpp b/Engine/Source/Runtime/Engine/Private/VT/RuntimeVirtualTexture.cpp
--- a/Engine/Source/Runtime/Engine/Private/VT/RuntimeVirtualTexture.cpp
+++ b/Engine/Source/Runtime/Engine/Private/VT/RuntimeVirtualTexture.cpp
@@ -403,6 +403,7 @@ void URuntimeVirtualTexture::GetAssetRegistryTags(TArray<FAssetRegistryTag>& Out
 	Super::GetAssetRegistryTags(OutTags);
 
 	OutTags.Add(FAssetRegistryTag("Size", FString::FromInt(GetSize()), FAssetRegistryTag::TT_Numerical));
+	OutTags.Add(FAssetRegistryTag("TileCount", FString::FromInt(GetTileCount()), FAssetRegistryTag::TT_Numerical));
 	OutTags.Add(FAssetRegistryTag("TileSize", FString::FromInt(GetTileSize()), FAssetRegistryTag::TT_Numerical));
 	OutTags.Add(FAssetRegistryTag("TileBorderSize", FString::FromInt(GetTileBorderSize()), FAssetRegistryTag::TT_Numerical));
 }
@@ -423,6 +424,20 @@ void URuntimeVirtualTexture::Serialize(FArchive& Ar)
 	}
 }
 
+void URuntimeVirtualTexture::PostLoad()
+{
+	// Convert Size_DEPRECATED to TileCount
+	if (Size_DEPRECATED >= 0)
+	{
+		int32 OldSize = 1 << FMath::Clamp(Size_DEPRECATED + 10, 10, 18);
+		int32 TileCountFromSize = FMath::Max(OldSize / GetTileSize(), 1);
+		TileCount = FMath::FloorLog2(TileCountFromSize);
+		Size_DEPRECATED = -1;
+	}
+
+	Super::PostLoad();
+}
+
 #if WITH_EDITOR
 
 void URuntimeVirtualTexture::PostEditChangeProperty(FPropertyChangedEvent& PropertyChangedEvent)
