Support baking RVT world height to streaming VT

diff --git a/Engine/Source/Editor/VirtualTexturingEditor/Private/RuntimeVirtualTextureBuild.cpp b/Engine/Source/Editor/VirtualTexturingEditor/Private/RuntimeVirtualTextureBuild.cpp
--- a/Engine/Source/Editor/VirtualTexturingEditor/Private/RuntimeVirtualTextureBuild.cpp
+++ b/Engine/Source/Editor/VirtualTexturingEditor/Private/RuntimeVirtualTextureBuild.cpp
@@ -5,10 +5,93 @@
 #include "Components/RuntimeVirtualTextureComponent.h"
 #include "Engine/TextureRenderTarget2D.h"
 #include "Misc/ScopedSlowTask.h"
+#include "RendererInterface.h"
+#include "RenderTargetPool.h"
 #include "SceneInterface.h"
 #include "VT/RuntimeVirtualTexture.h"
 #include "VT/RuntimeVirtualTextureRender.h"
 
+namespace
+{
+	/** Container for render resources needed to render the runtime virtual texture. */
+	class FRenderTileResources : public FRenderResource
+	{
+	public:
+		FRenderTileResources(int32 InNumLayers, int32 InTileSize, EPixelFormat InFormat)
+			: NumLayers(InNumLayers)
+			, TileSize(InTileSize)
+			, Format(InFormat)
+		{
+		}
+
+		//~ Begin FRenderResource Interface.
+		virtual void InitRHI() override
+		{
+			FRHICommandListImmediate& RHICmdList = FRHICommandListExecutor::GetImmediateCommandList();
+
+			RenderTargets.Init(nullptr, NumLayers);
+			StagingTextures.Init(nullptr, NumLayers);
+
+			for (int32 Layer = 0; Layer < NumLayers; ++Layer)
+			{
+				FRHIResourceCreateInfo CreateInfo;
+				RenderTargets[Layer] = RHICmdList.CreateTexture2D(TileSize, TileSize, Format, 1, 1, TexCreate_RenderTargetable, CreateInfo);
+				StagingTextures[Layer] = RHICmdList.CreateTexture2D(TileSize, TileSize, Format, 1, 1, TexCreate_CPUReadback, CreateInfo);
+			}
+
+			Fence = RHICmdList.CreateGPUFence(TEXT("Runtime Virtual Texture Build"));
+		}
+
+		virtual void ReleaseRHI() override
+		{
+			RenderTargets.Empty();
+			StagingTextures.Empty();
+			Fence.SafeRelease();
+		}
+		//~ End FRenderResource Interface.
+
+		FRHITexture2D* GetRenderTarget(int32 Index) const { return Index < NumLayers ? RenderTargets[Index] : nullptr; }
+		FRHITexture2D* GetStagingTexture(int32 Index) const { return Index < NumLayers ? StagingTextures[Index] : nullptr; }
+		FRHIGPUFence* GetFence() const { return Fence; }
+
+	private:
+		int32 NumLayers;
+		int32 TileSize;
+		EPixelFormat Format;
+
+		TArray<FTexture2DRHIRef> RenderTargets;
+		TArray<FTexture2DRHIRef> StagingTextures;
+		FGPUFenceRHIRef Fence;
+	};
+
+	/** Templatized helper function for copying a rendered tile to the final composited image data. */
+	template<typename T>
+	void TCopyTile(T* TilePixels, int32 TileSize, T* DestPixels, int32 DestStride, int32 DestLayerStride, FIntVector const& DestPos)
+	{
+		for (int32 y = 0; y < TileSize; y++)
+		{
+			memcpy(
+				DestPixels + DestLayerStride * DestPos[2] + DestStride * (DestPos[1] + y) + DestPos[0],
+				TilePixels + TileSize * y,
+				TileSize * sizeof(T));
+		}
+	}
+
+	/** Function for copying a rendered tile to the final composited image data. Needs ERuntimeVirtualTextureMaterialType to know what type of data is being copied. */
+	void CopyTile(void* TilePixels, int32 TileSize, void* DestPixels, int32 DestStride, int32 DestLayerStride, FIntVector const& DestPos, ERuntimeVirtualTextureMaterialType MaterialType)
+	{
+		if (MaterialType == ERuntimeVirtualTextureMaterialType::WorldHeight)
+		{
+			TCopyTile((uint16*)TilePixels, TileSize, (uint16*)DestPixels, DestStride, DestLayerStride, DestPos);
+		}
+		else
+		{
+			TCopyTile((FColor*)TilePixels, TileSize, (FColor*)DestPixels, DestStride, DestLayerStride, DestPos);
+		}
+	}
+}
+
+
 namespace RuntimeVirtualTexture
 {
 	bool HasStreamedMips(URuntimeVirtualTextureComponent* InComponent)
@@ -53,18 +136,15 @@ namespace RuntimeVirtualTexture
 		const int32 TextureSizeY = VTDesc.HeightInBlocks * VTDesc.BlockHeightInTiles * TileSize;
 		const int32 MaxLevel = (int32)FMath::CeilLogTwo(FMath::Max(VTDesc.BlockWidthInTiles, VTDesc.BlockHeightInTiles));
 		const int32 RenderLevel = FMath::Max(MaxLevel - RuntimeVirtualTexture->GetStreamLowMips() + 1, 0);
-		const int32 StreamingTextureSizeX = FMath::Max(TileSize, TextureSizeX >> RenderLevel);
-		const int32 StreamingTextureSizeY = FMath::Max(TileSize, TextureSizeY >> RenderLevel);
-		const int32 NumTilesX = StreamingTextureSizeX / TileSize;
-		const int32 NumTilesY = StreamingTextureSizeY / TileSize;
-		const int32 LayerCount = RuntimeVirtualTexture->GetLayerCount();
+		const int32 ImageSizeX = FMath::Max(TileSize, TextureSizeX >> RenderLevel);
+		const int32 ImageSizeY = FMath::Max(TileSize, TextureSizeY >> RenderLevel);
+		const int32 NumTilesX = ImageSizeX / TileSize;
+		const int32 NumTilesY = ImageSizeY / TileSize;
+		const int32 NumLayers = RuntimeVirtualTexture->GetLayerCount();
 
-		//todo[vt]: Support streaming VT for ERuntimeVirtualTextureMaterialType::WorldHeight (which requires UTextureRenderTarget2D support for PF_G16)
 		const ERuntimeVirtualTextureMaterialType MaterialType = RuntimeVirtualTexture->GetMaterialType();
-		if (MaterialType == ERuntimeVirtualTextureMaterialType::WorldHeight)
-		{
-			return false;
-		}
+		const EPixelFormat RenderTargetFormat = (MaterialType == ERuntimeVirtualTextureMaterialType::WorldHeight) ? PF_G16 : PF_B8G8R8A8;
+		const int32 BytesPerPixel = (MaterialType == ERuntimeVirtualTextureMaterialType::WorldHeight) ? 2 : 4;
 
 		// Spin up slow task UI
 		const float TaskWorkRender = NumTilesX * NumTilesY;
@@ -72,77 +152,97 @@ namespace RuntimeVirtualTexture
 		FScopedSlowTask Task(TaskWorkRender + TaskWorkBuildBulkData, FText::AsCultureInvariant(RuntimeVirtualTexture->GetName()));
 		Task.MakeDialog(true);
 
-		// Final pixels will contain data for each virtual texture layer in order
-		TArray64<FColor> FinalPixels;
-		FinalPixels.InsertUninitialized(0, StreamingTextureSizeX * StreamingTextureSizeY * LayerCount);
-
 		// Allocate render targets for rendering out the runtime virtual texture tiles
-		UTextureRenderTarget2D* RenderTarget[2] = { nullptr };
-		FRenderTarget* RenderTargetResource[2] = { nullptr };
+		FRenderTileResources RenderTileResources(NumLayers, TileSize, RenderTargetFormat);
+		BeginInitResource(&RenderTileResources);
 
-		for (int32 Layer = 0; Layer < LayerCount; Layer++)
-		{
-			RenderTarget[Layer] = NewObject<UTextureRenderTarget2D>();
-			RenderTarget[Layer]->AddToRoot();
-			RenderTarget[Layer]->ClearColor = FLinearColor(0.0f, 0.0f, 0.0f, 0.0f);
-			RenderTarget[Layer]->InitCustomFormat(TileSize, TileSize, PF_B8G8R8A8, false);
-			RenderTargetResource[Layer] = RenderTarget[Layer]->GameThread_GetRenderTargetResource();
-		}
+		// Final pixels will contain image data for each virtual texture layer in order
+		TArray64<uint8> FinalPixels;
+		FinalPixels.InsertUninitialized(0, ImageSizeX * ImageSizeY * NumLayers * BytesPerPixel);
 
-		// Run over all tiles and render/store each one
-		for (int32 Y = 0; Y < NumTilesY && !Task.ShouldCancel(); Y++)
+		// Iterate over all tiles and render/store each one to the final image
+		for (int32 TileY = 0; TileY < NumTilesY && !Task.ShouldCancel(); TileY++)
 		{
-			for (int32 X = 0; X < NumTilesX && !Task.ShouldCancel(); X++)
+			for (int32 TileX = 0; TileX < NumTilesX && !Task.ShouldCancel(); TileX++)
 			{
 				Task.EnterProgressFrame();
 
 				// Render tile
-				const FBox2D TileBox = FBox2D(FVector2D(0, 0), FVector2D(TileSize, TileSize));
 				const FBox2D UVRange = FBox2D(
-					FVector2D((float)X / (float)NumTilesX, (float)Y / (float)NumTilesY),
-					FVector2D((float)(X + 1) / (float)NumTilesX, (float)(Y + 1) / (float)NumTilesY));
+					FVector2D((float)TileX / (float)NumTilesX, (float)TileY / (float)NumTilesY),
+					FVector2D((float)(TileX + 1) / (float)NumTilesX, (float)(TileY + 1) / (float)NumTilesY));
 
-				ENQUEUE_RENDER_COMMAND(BakeStreamingTextureTileCommand)(
-					[Scene, VirtualTextureSceneIndex, MaterialType, RenderTargetResource, TileBox, Transform, UVRange, RenderLevel, MaxLevel, DebugType](FRHICommandListImmediate& RHICmdList)
+				ENQUEUE_RENDER_COMMAND(BakeStreamingTextureTileCommand)([
+					Scene, VirtualTextureSceneIndex, 
+					&RenderTileResources,
+					MaterialType, NumLayers,
+					Transform, UVRange,
+					RenderLevel, MaxLevel, 
+					TileX, TileY,
+					TileSize, ImageSizeX, ImageSizeY, 
+					&FinalPixels,
+					DebugType
+				](FRHICommandListImmediate& RHICmdList)
 				{
+					const FBox2D TileBox(FVector2D(0, 0), FVector2D(TileSize, TileSize));
+					const FIntRect TileRect(0, 0, TileSize, TileSize);
+
+					// Transition render targets for writing
+					for (int32 Layer = 0; Layer < NumLayers; Layer++)
+					{
+						RHICmdList.TransitionResource(EResourceTransitionAccess::EWritable, RenderTileResources.GetRenderTarget(Layer));
+					}
+
+					// Render a single tile
 					RuntimeVirtualTexture::RenderPage(
 						RHICmdList,
 						Scene->GetRenderScene(),
 						1 << VirtualTextureSceneIndex,
 						MaterialType,
 						true,
-						RenderTargetResource[0] != nullptr ? RenderTargetResource[0]->GetRenderTargetTexture() : nullptr, TileBox,
-						RenderTargetResource[1] != nullptr ? RenderTargetResource[1]->GetRenderTargetTexture() : nullptr, TileBox,
+						RenderTileResources.GetRenderTarget(0), TileBox,
+						RenderTileResources.GetRenderTarget(1), TileBox,
 						Transform,
 						UVRange,
 						RenderLevel,
 						MaxLevel,
 						DebugType);
-				});
 
-				FlushRenderingCommands();
+					// Transition render targets for copying
+					for (int32 Layer = 0; Layer < NumLayers; Layer++)
+					{
+						RHICmdList.TransitionResource(EResourceTransitionAccess::EReadable, RenderTileResources.GetRenderTarget(Layer));
+					}
 
-				// Read back into final pixel data
-				for (int32 Layer = 0; Layer < LayerCount; Layer++)
-				{
-					TArray<FColor> TilePixels;
-					RenderTarget[Layer]->GameThread_GetRenderTargetResource()->ReadPixels(TilePixels);
+					// Copy to staging
+					for (int32 Layer = 0; Layer < NumLayers; Layer++)
+					{
+						RHICmdList.CopyTexture(RenderTileResources.GetRenderTarget(Layer), RenderTileResources.GetStagingTexture(Layer), FRHICopyTextureInfo());
+					}
+
+					//todo[vt]: Insert fence for immediate read back. But is there no API to wait on it?
+					RHICmdList.WriteGPUFence(RenderTileResources.GetFence());
+					RHICmdList.ImmediateFlush(EImmediateFlushType::FlushRHIThread);
 
-					for (int32 y = 0; y < TileSize; y++)
+					// Read back tile data and copy into final destination
+					for (int32 Layer = 0; Layer < NumLayers; Layer++)
 					{
-						for (int32 x = 0; x < TileSize; x++)
-						{
-							FinalPixels[Layer * StreamingTextureSizeX * StreamingTextureSizeY + (Y * TileSize + y) * StreamingTextureSizeX + (X * TileSize + x)] = TilePixels[y * TileSize + x];
-						}
+						void* TilePixels = nullptr;
+						int32 OutWidth, OutHeight;
+						RHICmdList.MapStagingSurface(RenderTileResources.GetStagingTexture(Layer), TilePixels, OutWidth, OutHeight);
+						check(TilePixels != nullptr);
+						check(OutWidth == TileSize && OutHeight == TileSize);
+
+						const FIntVector DestPos(TileX * TileSize, TileY * TileSize, Layer);
+						CopyTile(TilePixels, TileSize, FinalPixels.GetData(), ImageSizeX, ImageSizeX * ImageSizeY, DestPos, MaterialType);
+
+						RHICmdList.UnmapStagingSurface(RenderTileResources.GetStagingTexture(Layer));
 					}
-				}
+				});
 			}
 		}
 
-		for (int32 Layer = 0; Layer < LayerCount; Layer++)
-		{
-			RenderTarget[Layer]->RemoveFromRoot();
-		}
+		ReleaseResourceAndFlush(&RenderTileResources);
 
 		if (Task.ShouldCancel())
 		{
@@ -152,7 +252,7 @@ namespace RuntimeVirtualTexture
 		// Place final pixel data into the runtime virtual texture
 		Task.EnterProgressFrame(TaskWorkBuildBulkData);
 		RuntimeVirtualTexture->Modify();
-		RuntimeVirtualTexture->InitializeStreamingTexture(StreamingTextureSizeX, StreamingTextureSizeY, (uint8*)FinalPixels.GetData());
+		RuntimeVirtualTexture->InitializeStreamingTexture(ImageSizeX, ImageSizeY, (uint8*)FinalPixels.GetData());
 		RuntimeVirtualTexture->PostEditChange();
 
 		return true;
diff --git a/Engine/Source/Runtime/Engine/Private/VT/RuntimeVirtualTexture.cpp b/Engine/Source/Runtime/Engine/Private/VT/RuntimeVirtualTexture.cpp
--- a/Engine/Source/Runtime/Engine/Private/VT/RuntimeVirtualTexture.cpp
+++ b/Engine/Source/Runtime/Engine/Private/VT/RuntimeVirtualTexture.cpp
@@ -491,20 +491,24 @@ void URuntimeVirtualTexture::InitializeStreamingTexture(uint32 InSizeX, uint32 I
 
 	StreamingTexture->BuildHash = GetStreamingTextureBuildHash();
 
+	enum { MAX_RVT_LAYERS = 2 };
 	const int32 LayerCount = GetLayerCount();
+	check(LayerCount <= MAX_RVT_LAYERS);
+	ETextureSourceFormat LayerFormats[MAX_RVT_LAYERS];
+
 	for (int32 Layer = 0; Layer < LayerCount; Layer++)
 	{
 		EPixelFormat LayerFormat = GetLayerFormat(Layer);
+		LayerFormats[Layer] = LayerFormat == PF_G16 ? TSF_G16 : TSF_BGRA8;
 
 		FTextureFormatSettings FormatSettings;
 		FormatSettings.SRGB = IsLayerSRGB(Layer);
-		FormatSettings.CompressionNone = !bCompressTextures;
+		FormatSettings.CompressionNone = LayerFormat == PF_B8G8R8A8 ||LayerFormat == PF_G16;
 		FormatSettings.CompressionNoAlpha = LayerFormat == PF_DXT1 || LayerFormat == PF_BC5;
 		FormatSettings.CompressionSettings = LayerFormat == PF_BC5 ? TC_Normalmap : TC_Default;
 		StreamingTexture->SetLayerFormatSettings(Layer, FormatSettings);
 	}
 
-	const ETextureSourceFormat LayerFormats[] = { TSF_BGRA8, TSF_BGRA8 };
 	StreamingTexture->Source.InitLayered(InSizeX, InSizeY, 1, LayerCount, 1, LayerFormats, InData);
 
 	StreamingTexture->PostEditChange();
@@ -519,9 +523,15 @@ IVirtualTexture* URuntimeVirtualTexture::CreateStreamingTextureProducer(IVirtual
 		FTexturePlatformData** StreamingTextureData = StreamingTexture->GetRunningPlatformData();
 		if (StreamingTextureData != nullptr && *StreamingTextureData != nullptr)
 		{
-			int32 NumStreamMips = FMath::Min(GetStreamLowMips(), (*StreamingTextureData)->GetNumVTMips()); // Min() is a hack while we fix crash due to data breakage
+			FVirtualTextureBuiltData* VTData = (*StreamingTextureData)->VTData;
+			check(GetTileSize() == VTData->TileSize);
+			check(GetTileBorderSize() == VTData->TileBorderSize);
+
+			// Streaming data may have mips removed during cook
+			const int32 NumStreamMips = FMath::Min(GetStreamLowMips(), (*StreamingTextureData)->GetNumVTMips());
+
 			OutTransitionLevel = FMath::Max(InMaxLevel - NumStreamMips + 1, 0);
-			IVirtualTexture* StreamingProducer = new FUploadingVirtualTexture((*StreamingTextureData)->VTData, 0);
+			IVirtualTexture* StreamingProducer = new FUploadingVirtualTexture(VTData, 0);
 			return new FVirtualTextureLevelRedirector(InProducer, StreamingProducer, OutTransitionLevel);
 		}
 	}
