Added runtime virtual texture asset action "Fix Material Usage"
This fixes all dependent materials to match the selected runtime virtual texture asset material layout
Useful when changing  RVT layout in a big project

diff --git a/Engine/Source/Editor/VirtualTexturingEditor/Classes/RuntimeVirtualTextureAssetTypeActions.h b/Engine/Source/Editor/VirtualTexturingEditor/Classes/RuntimeVirtualTextureAssetTypeActions.h
--- a/Engine/Source/Editor/VirtualTexturingEditor/Classes/RuntimeVirtualTextureAssetTypeActions.h
+++ b/Engine/Source/Editor/VirtualTexturingEditor/Classes/RuntimeVirtualTextureAssetTypeActions.h
@@ -5,6 +5,8 @@
 #include "CoreMinimal.h"
 #include "AssetTypeActions_Base.h"
 
+class URuntimeVirtualTexture;
+
 /** Asset actions setup for URuntimeVirtualTexture */
 class FAssetTypeActions_RuntimeVirtualTexture : public FAssetTypeActions_Base
 {
@@ -17,5 +19,13 @@ protected:
 	virtual FText GetName() const override;
 	virtual FColor GetTypeColor() const override;
 	virtual uint32 GetCategories() override;
+	virtual bool HasActions(TArray<UObject*> const& InObjects) const override { return true; }
+	virtual void GetActions(TArray<UObject*> const& InObjects, FMenuBuilder& MenuBuilder) override;
 	//~ End FAssetTypeActions_Base Interface.
+
+private:
+	/** Handler for when FindMaterials is selected */
+	void ExecuteFindMaterials(TWeakObjectPtr<URuntimeVirtualTexture> Object);
+	/** Handler for when FixMaterialUsage is selected */
+	void ExecuteFixMaterialUsage(TWeakObjectPtr<URuntimeVirtualTexture> Object);
 };
diff --git a/Engine/Source/Editor/VirtualTexturingEditor/Private/RuntimeVirtualTextureAssetTypeActions.cpp b/Engine/Source/Editor/VirtualTexturingEditor/Private/RuntimeVirtualTextureAssetTypeActions.cpp
--- a/Engine/Source/Editor/VirtualTexturingEditor/Private/RuntimeVirtualTextureAssetTypeActions.cpp
+++ b/Engine/Source/Editor/VirtualTexturingEditor/Private/RuntimeVirtualTextureAssetTypeActions.cpp
@@ -2,11 +2,189 @@
 
 #include "RuntimeVirtualTextureAssetTypeActions.h"
 
+#include "AssetRegistryModule.h"
+#include "ContentBrowserModule.h"
+#include "EditorSupportDelegates.h"
+#include "Framework/Commands/UIAction.h"
+#include "Framework/MultiBox/MultiBoxBuilder.h"
+#include "IAssetRegistry.h"
+#include "IContentBrowserSingleton.h"
+#include "MaterialEditingLibrary.h"
+#include "Materials/Material.h"
+#include "Materials/MaterialExpressionRuntimeVirtualTextureSample.h"
+#include "Materials/MaterialFunction.h"
+#include "Materials/MaterialInstance.h"
+#include "Materials/MaterialInterface.h"
+#include "Misc/ScopedSlowTask.h"
+#include "Modules/ModuleManager.h"
 #include "Toolkits/IToolkitHost.h"
 #include "VT/RuntimeVirtualTexture.h"
 
 #define LOCTEXT_NAMESPACE "VirtualTexturingEditorModule"
 
+DEFINE_LOG_CATEGORY_STATIC(LogRuntimeVirtualTextureFixMaterial, Log, Log);
+
+namespace
+{
+	void GetReferencersData(UObject *Object, UClass *MatchClass, TArray<FAssetData> &OutAssetDatas)
+	{
+		FAssetRegistryModule& AssetRegistryModule = FModuleManager::LoadModuleChecked<FAssetRegistryModule>("AssetRegistry");
+		IAssetRegistry& AssetRegistry = AssetRegistryModule.Get();
+
+		TArray<FAssetIdentifier> Referencers;
+		AssetRegistry.GetReferencers(Object->GetOuter()->GetFName(), Referencers);
+
+		for (auto AssetIdentifier : Referencers)
+		{
+			TArray<FAssetData> Assets;
+			AssetRegistry.GetAssetsByPackageName(AssetIdentifier.PackageName, Assets);
+
+			for (auto AssetData : Assets)
+			{
+				if (MatchClass != nullptr)
+				{
+					if (AssetData.GetClass()->IsChildOf(MatchClass))
+					{
+						OutAssetDatas.AddUnique(AssetData);
+					}
+				}
+				else
+				{
+					OutAssetDatas.AddUnique(AssetData);
+				}
+			}
+		}
+	}
+
+	template <class T> 
+	void GetReferencersOfType(UObject *Object, TArray<T*> &OutObjects)
+	{
+		TArray<FAssetData> AssetDatas;
+		GetReferencersData(Object, T::StaticClass(), AssetDatas);
+
+		for (auto Data : AssetDatas)
+		{
+			T* TypedAsset = Cast<T>(Data.GetAsset());
+			if (TypedAsset != nullptr)
+			{
+				OutObjects.Add(TypedAsset);
+			}
+		}
+	}
+
+	void FindAllMaterials(URuntimeVirtualTexture* RuntimeVirtualTexture, TArray<UMaterial*>& OutMaterials, TArray<UMaterialFunctionInterface*>& OutFunctions)
+	{
+		TArray<UMaterialInterface*> MaterialInterfaces;
+		GetReferencersOfType(RuntimeVirtualTexture, MaterialInterfaces);
+
+		for (UMaterialInterface *MaterialInterface : MaterialInterfaces)
+		{
+			UMaterial* Material = Cast<UMaterial>(MaterialInterface);
+			
+			UMaterialInstance* MaterialInstance = Cast<UMaterialInstance>(MaterialInterface);
+			if (MaterialInstance != nullptr)
+			{
+				Material = MaterialInstance->GetMaterial();
+			}
+
+			OutMaterials.AddUnique(Material);
+		}
+
+		GetReferencersOfType(RuntimeVirtualTexture, OutFunctions);
+	}
+
+	void FixMaterialUsage(URuntimeVirtualTexture* RuntimeVirtualTexture)
+	{
+		UE_LOG(LogRuntimeVirtualTextureFixMaterial, Log, TEXT("Begin fix material usage for '%s' ..."), *RuntimeVirtualTexture->GetName());
+
+		TArray<UMaterial*> Materials;
+		TArray<UMaterialFunctionInterface*> Functions;
+		FindAllMaterials(RuntimeVirtualTexture, Materials, Functions);
+
+		int32 TaskCount = Materials.Num() + Functions.Num();
+		FScopedSlowTask Task(TaskCount, LOCTEXT("RuntimeVirtualTexture_FixMaterialUsageProgress", "Fixing materials for Runtime Virtual Texture usage..."));
+		Task.MakeDialog();
+
+		for (UMaterial* Material : Materials)
+		{
+			Task.EnterProgressFrame();
+
+			bool bMaterialModified = false;
+			for (UMaterialExpression* Expression : Material->Expressions)
+			{
+				UMaterialExpressionRuntimeVirtualTextureSample* RVTSampleExpression = Cast<UMaterialExpressionRuntimeVirtualTextureSample>(Expression);
+				if (RVTSampleExpression)
+				{
+					if (RuntimeVirtualTexture == RVTSampleExpression->VirtualTexture)
+					{
+						if (RVTSampleExpression->InitVirtualTextureDependentSettings())
+						{
+							Expression->Modify();
+
+							FPropertyChangedEvent Event(UMaterialExpressionTextureBase::StaticClass()->FindPropertyByName(GET_MEMBER_NAME_CHECKED(UMaterialExpressionRuntimeVirtualTextureSample, MaterialType)));
+							Expression->PostEditChangeProperty(Event);
+
+							bMaterialModified = true;
+						}
+					}
+				}
+			}
+
+			if (bMaterialModified)
+			{
+				UE_LOG(LogRuntimeVirtualTextureFixMaterial, Log, TEXT("  Recompile material '%s' ..."), *Material->GetName());
+
+				FScopedSlowTask CompileTask(1, FText::AsCultureInvariant(Material->GetName()));
+				CompileTask.MakeDialog();
+				CompileTask.EnterProgressFrame();
+
+				UMaterialEditingLibrary::RecompileMaterial(Material);
+			}
+		}
+
+		for (UMaterialFunctionInterface *Function : Functions)
+		{
+			Task.EnterProgressFrame();
+
+			bool bFunctionModified = false;
+			const TArray<UMaterialExpression*> *Expressions = Function->GetFunctionExpressions();
+			for (UMaterialExpression *Expression : *Expressions)
+			{
+				UMaterialExpressionRuntimeVirtualTextureSample* RVTSampleExpression = Cast<UMaterialExpressionRuntimeVirtualTextureSample>(Expression);
+				if (RVTSampleExpression)
+				{
+					if (RuntimeVirtualTexture == RVTSampleExpression->VirtualTexture)
+					{
+						if (RVTSampleExpression->InitVirtualTextureDependentSettings())
+						{
+							Expression->Modify();
+
+							FPropertyChangedEvent Event(UMaterialExpressionTextureBase::StaticClass()->FindPropertyByName(GET_MEMBER_NAME_CHECKED(UMaterialExpressionRuntimeVirtualTextureSample, MaterialType)));
+							Expression->PostEditChangeProperty(Event);
+
+							bFunctionModified = true;
+						}
+					}
+				}
+			}
+
+			if (bFunctionModified)
+			{
+				UE_LOG(LogRuntimeVirtualTextureFixMaterial, Log, TEXT("  Update function '%s' ..."), *Function->GetName());
+
+				FScopedSlowTask CompileTask(1, FText::AsCultureInvariant(Function->GetName()));
+				CompileTask.MakeDialog();
+				CompileTask.EnterProgressFrame();
+
+				UMaterialEditingLibrary::UpdateMaterialFunction(Function, nullptr);
+			}
+		}
+
+		UE_LOG(LogRuntimeVirtualTextureFixMaterial, Log, TEXT("End fix material usage for '%s' ..."), *RuntimeVirtualTexture->GetName());
+	}
+}
+
+
 UClass* FAssetTypeActions_RuntimeVirtualTexture::GetSupportedClass() const
 {
 	return URuntimeVirtualTexture::StaticClass();
@@ -27,4 +205,62 @@ uint32 FAssetTypeActions_RuntimeVirtualTexture::GetCategories()
 	return EAssetTypeCategories::MaterialsAndTextures; 
 }
 
+void FAssetTypeActions_RuntimeVirtualTexture::GetActions(TArray<UObject*> const& InObjects, FMenuBuilder& MenuBuilder)
+{
+	if (InObjects.Num() == 1)
+	{
+		TArray<TWeakObjectPtr<URuntimeVirtualTexture>> RuntimeVirtualTextures = GetTypedWeakObjectPtrs<URuntimeVirtualTexture>(InObjects);
+
+		MenuBuilder.AddMenuEntry(
+			LOCTEXT("RuntimeVirtualTexture_FindMaterials", "Find Materials Using This"),
+			LOCTEXT("RuntimeVirtualTexture_FindMaterialsTooltip", "Finds all materials that use this material in the content browser."),
+			FSlateIcon(),
+			FUIAction(
+				FExecuteAction::CreateSP(this, &FAssetTypeActions_RuntimeVirtualTexture::ExecuteFindMaterials, RuntimeVirtualTextures[0]),
+				FCanExecuteAction()
+			)
+		);
+
+		MenuBuilder.AddMenuEntry(
+			LOCTEXT("RuntimeVirtualTexture_FixMaterialUsage", "Fix Material Usage"),
+			LOCTEXT("RuntimeVirtualTexture_FixMaterialUsageTooltip", "Find materials using this Runtime Virtual Texture and fix any mismatching content types."),
+			FSlateIcon(FEditorStyle::GetStyleSetName(), "ClassIcon.Texture2D"),
+			FUIAction(
+				FExecuteAction::CreateSP(this, &FAssetTypeActions_RuntimeVirtualTexture::ExecuteFixMaterialUsage, RuntimeVirtualTextures[0]),
+				FCanExecuteAction()
+			)
+		);
+	}
+}
+
+void FAssetTypeActions_RuntimeVirtualTexture::ExecuteFindMaterials(TWeakObjectPtr<URuntimeVirtualTexture> Object)
+{
+	TArray<FAssetData> Materials;
+
+	URuntimeVirtualTexture* RuntimeVirtualTexture = Object.Get();
+	if (RuntimeVirtualTexture != nullptr)
+	{
+		GetReferencersData(RuntimeVirtualTexture, UMaterialInterface::StaticClass(), Materials);
+		GetReferencersData(RuntimeVirtualTexture, UMaterialFunction::StaticClass(), Materials);
+	}
+
+	if (Materials.Num() > 0)
+	{
+		FContentBrowserModule& ContentBrowserModule = FModuleManager::LoadModuleChecked<FContentBrowserModule>("ContentBrowser");
+		ContentBrowserModule.Get().SyncBrowserToAssets(Materials);
+	}
+}
+
+void FAssetTypeActions_RuntimeVirtualTexture::ExecuteFixMaterialUsage(TWeakObjectPtr<URuntimeVirtualTexture> Object)
+{
+	URuntimeVirtualTexture* RuntimeVirtualTexture = Object.Get();
+	if (RuntimeVirtualTexture != nullptr)
+	{
+		FixMaterialUsage(RuntimeVirtualTexture);
+
+		FEditorDelegates::RefreshEditor.Broadcast();
+		FEditorSupportDelegates::RedrawAllViewports.Broadcast();
+	}
+}
+
 #undef LOCTEXT_NAMESPACE
diff --git a/Engine/Source/Editor/VirtualTexturingEditor/VirtualTexturingEditor.Build.cs b/Engine/Source/Editor/VirtualTexturingEditor/VirtualTexturingEditor.Build.cs
--- a/Engine/Source/Editor/VirtualTexturingEditor/VirtualTexturingEditor.Build.cs
+++ b/Engine/Source/Editor/VirtualTexturingEditor/VirtualTexturingEditor.Build.cs
@@ -32,6 +32,7 @@ public class VirtualTexturingEditor : ModuleRules
 				"EditorStyle",
 				"Engine",
 				"InputCore",
+				"MaterialEditor",
 				"PlacementMode",
 				"PropertyEditor",
 				"RenderCore",
diff --git a/Engine/Source/Runtime/Engine/Classes/Materials/MaterialExpressionRuntimeVirtualTextureSample.h b/Engine/Source/Runtime/Engine/Classes/Materials/MaterialExpressionRuntimeVirtualTextureSample.h
--- a/Engine/Source/Runtime/Engine/Classes/Materials/MaterialExpressionRuntimeVirtualTextureSample.h
+++ b/Engine/Source/Runtime/Engine/Classes/Materials/MaterialExpressionRuntimeVirtualTextureSample.h
@@ -58,10 +58,10 @@ class ENGINE_API UMaterialExpressionRuntimeVirtualTextureSample : public UMateri
 	UPROPERTY(EditAnywhere, Category = TextureSample)
 	TEnumAsByte<enum ERuntimeVirtualTextureMipValueMode> MipValueMode = RVTMVM_None;
 
-protected:
 	/** Init settings that affect shader compilation and need to match the current VirtualTexture */
-	void InitVirtualTextureDependentSettings();
+	bool InitVirtualTextureDependentSettings();
 
+protected:
 	/** Initialize the output pins. */
 	void InitOutputs();
 
