Fixed compile errors due to missing header includes and forward declarations

diff --git a/Engine/Source/Editor/VirtualTexturingEditor/Classes/RuntimeVirtualTextureDetailsCustomization.h b/Engine/Source/Editor/VirtualTexturingEditor/Classes/RuntimeVirtualTextureDetailsCustomization.h
--- a/Engine/Source/Editor/VirtualTexturingEditor/Classes/RuntimeVirtualTextureDetailsCustomization.h
+++ b/Engine/Source/Editor/VirtualTexturingEditor/Classes/RuntimeVirtualTextureDetailsCustomization.h
@@ -4,6 +4,8 @@
 
 #include "CoreMinimal.h"
 #include "IDetailCustomization.h"
+#include "Widgets/Text/STextBlock.h"
+#include "Input/Reply.h"
 
 /** UI customization for URuntimeVirtualTexture */
 class FRuntimeVirtualTextureDetailsCustomization : public IDetailCustomization
diff --git a/Engine/Source/Editor/VirtualTexturingEditor/Private/RuntimeVirtualTextureDetailsCustomization.cpp b/Engine/Source/Editor/VirtualTexturingEditor/Private/RuntimeVirtualTextureDetailsCustomization.cpp
--- a/Engine/Source/Editor/VirtualTexturingEditor/Private/RuntimeVirtualTextureDetailsCustomization.cpp
+++ b/Engine/Source/Editor/VirtualTexturingEditor/Private/RuntimeVirtualTextureDetailsCustomization.cpp
@@ -9,6 +9,8 @@
 #include "SResetToDefaultMenu.h"
 #include "VT/RuntimeVirtualTexture.h"
 #include "Widgets/Layout/SWrapBox.h"
+#include "Widgets/Layout/SBox.h"
+#include "Widgets/Input/SButton.h"
 
 #define LOCTEXT_NAMESPACE "VirtualTexturingEditorModule"
 
diff --git a/Engine/Source/Runtime/Engine/Private/VT/VirtualTextureTranscodeCache.h b/Engine/Source/Runtime/Engine/Private/VT/VirtualTextureTranscodeCache.h
--- a/Engine/Source/Runtime/Engine/Private/VT/VirtualTextureTranscodeCache.h
+++ b/Engine/Source/Runtime/Engine/Private/VT/VirtualTextureTranscodeCache.h
@@ -9,6 +9,7 @@
 
 class FUploadingVirtualTexture;
 class FVirtualTextureCodec;
+struct FVirtualTextureBuiltData;
 
 struct FVTTranscodeParams
 {
diff --git a/Engine/Source/Runtime/RenderCore/Public/VirtualTexturing.h b/Engine/Source/Runtime/RenderCore/Public/VirtualTexturing.h
--- a/Engine/Source/Runtime/RenderCore/Public/VirtualTexturing.h
+++ b/Engine/Source/Runtime/RenderCore/Public/VirtualTexturing.h
@@ -5,6 +5,7 @@
 #include "CoreMinimal.h"
 #include "RHIDefinitions.h"
 #include "Templates/RefCounting.h"
+#include "Stats/Stats.h"
 
 class FRHICommandListImmediate;
 class FRHIShaderResourceView;
