Expose placement mode category refresh event.
Use it to add virtual texture volume actor to volumes category.
Need to change VirtualTexturingEditor module load order for this.
Move ActorFactoryRuntimeVirtualTextureVolume to account for this later load.
#rnx

diff --git a/Engine/Source/Editor/PlacementMode/Private/PlacementModeModule.cpp b/Engine/Source/Editor/PlacementMode/Private/PlacementModeModule.cpp
--- a/Engine/Source/Editor/PlacementMode/Private/PlacementModeModule.cpp
+++ b/Engine/Source/Editor/PlacementMode/Private/PlacementModeModule.cpp
@@ -415,7 +415,7 @@ void FPlacementModeModule::RegenerateItemsForCategory(FName Category)
 		RefreshAllPlaceableClasses();
 	}
 
-	BroadcastPlacementModeCategoryRefreshed(Category);
+	PlacementModeCategoryRefreshed.Broadcast(Category);
 }
 
 void FPlacementModeModule::RefreshRecentlyPlaced()
diff --git a/Engine/Source/Editor/PlacementMode/Public/IPlacementModeModule.h b/Engine/Source/Editor/PlacementMode/Public/IPlacementModeModule.h
--- a/Engine/Source/Editor/PlacementMode/Public/IPlacementModeModule.h
+++ b/Engine/Source/Editor/PlacementMode/Public/IPlacementModeModule.h
@@ -217,6 +217,12 @@ public:
 	 */
 	virtual const TArray< FActorPlacementInfo >& GetRecentlyPlaced() const = 0;
 
+	/**
+	 * @return the event that is broadcast whenever a placement mode category is refreshed
+	 */
+	DECLARE_EVENT_OneParam( IPlacementMode, FOnPlacementModeCategoryRefreshed, FName /*CategoryName*/ );
+	virtual FOnPlacementModeCategoryRefreshed& OnPlacementModeCategoryRefreshed() = 0;
+
 	/**
 	 * @return the event that is broadcast whenever the list of recently placed assets changes
 	 */
diff --git a/Engine/Source/Editor/PlacementMode/Public/PlacementModeModule.h b/Engine/Source/Editor/PlacementMode/Public/PlacementModeModule.h
--- a/Engine/Source/Editor/PlacementMode/Public/PlacementModeModule.h
+++ b/Engine/Source/Editor/PlacementMode/Public/PlacementModeModule.h
@@ -7,8 +7,6 @@
 
 #include "Framework/MultiBox/MultiBoxExtender.h"
 
-DECLARE_MULTICAST_DELEGATE_OneParam(FOnPlacementModeCategoryRefreshed, FName /*CategoryName*/)
-
 struct FPlacementCategory : FPlacementCategoryInfo
 {
 	FPlacementCategory(const FPlacementCategoryInfo& SourceInfo)
@@ -48,16 +46,15 @@ public:
 	 */
 	virtual void PreUnloadCallback() override;
 
+	DECLARE_DERIVED_EVENT(FPlacementModeModule, IPlacementModeModule::FOnPlacementModeCategoryRefreshed, FOnPlacementModeCategoryRefreshed);
+	virtual FOnPlacementModeCategoryRefreshed& OnPlacementModeCategoryRefreshed() override { return PlacementModeCategoryRefreshed; }
+
 	DECLARE_DERIVED_EVENT(FPlacementModeModule, IPlacementModeModule::FOnRecentlyPlacedChanged, FOnRecentlyPlacedChanged);
 	virtual FOnRecentlyPlacedChanged& OnRecentlyPlacedChanged() override { return RecentlyPlacedChanged; }
 
 	DECLARE_DERIVED_EVENT(FPlacementModeModule, IPlacementModeModule::FOnAllPlaceableAssetsChanged, FOnAllPlaceableAssetsChanged);
 	virtual FOnAllPlaceableAssetsChanged& OnAllPlaceableAssetsChanged() override { return AllPlaceableAssetsChanged; }
 
-	FOnPlacementModeCategoryRefreshed& OnPlacementModeCategoryRefreshed() { return PlacementModeCategoryRefreshed; }
-
-	void BroadcastPlacementModeCategoryRefreshed(FName CategoryName) { PlacementModeCategoryRefreshed.Broadcast(CategoryName); }
-
 	/**
 	 * Add the specified assets to the recently placed items list
 	 */
diff --git a/Engine/Source/Editor/UnrealEd/Classes/ActorFactories/ActorFactoryRuntimeVirtualTextureVolume.h b/Engine/Source/Editor/UnrealEd/Classes/ActorFactories/ActorFactoryRuntimeVirtualTextureVolume.h
new file mode 100644
--- /dev/null
+++ b/Engine/Source/Editor/UnrealEd/Classes/ActorFactories/ActorFactoryRuntimeVirtualTextureVolume.h
@@ -0,0 +1,18 @@
+// Copyright 1998-2019 Epic Games, Inc. All Rights Reserved.
+
+#pragma once
+
+#include "CoreMinimal.h"
+#include "ActorFactories/ActorFactory.h"
+#include "ActorFactoryRuntimeVirtualTextureVolume.generated.h"
+
+/** Actor factory for ARuntimeVirtualTextureVolume */
+UCLASS(MinimalAPI, config=Editor, collapsecategories, hidecategories=Object)
+class UActorFactoryRuntimeVirtualTextureVolume : public UActorFactory
+{
+	GENERATED_UCLASS_BODY()
+
+	//~ Begin UActorFactory Interface
+	virtual void PostSpawnActor(UObject* Asset, AActor* NewActor) override;
+	//~ End UActorFactory Interface
+};
diff --git a/Engine/Source/Editor/UnrealEd/Private/EditorModeRegistry.cpp b/Engine/Source/Editor/UnrealEd/Private/EditorModeRegistry.cpp
--- a/Engine/Source/Editor/UnrealEd/Private/EditorModeRegistry.cpp
+++ b/Engine/Source/Editor/UnrealEd/Private/EditorModeRegistry.cpp
@@ -16,6 +16,7 @@
 #include "Editor/SceneDepthPickerMode/Public/SceneDepthPickerMode.h"
 #include "Editor/TextureAlignMode/Public/TextureAlignEdMode.h"
 #include "Editor/FoliageEdit/Public/FoliageEditModule.h"
+#include "Editor/VirtualTexturingEditor/Public/VirtualTexturingEditorModule.h"
 
 FEditorModeInfo::FEditorModeInfo()
 	: ID(NAME_None)
@@ -65,6 +66,7 @@ void FEditorModeRegistry::Initialize()
 	FModuleManager::LoadModuleChecked<IMeshPaintModule>(TEXT("MeshPaintMode"));
 	FModuleManager::LoadModuleChecked<ILandscapeEditorModule>(TEXT("LandscapeEditor"));
 	FModuleManager::LoadModuleChecked<IFoliageEditModule>(TEXT("FoliageEdit"));
+	FModuleManager::LoadModuleChecked<FVirtualTexturingEditorModule>(TEXT("VirtualTexturingEditor"));
 }
 
 void FEditorModeRegistry::Shutdown()
diff --git a/Engine/Source/Editor/UnrealEd/Private/Factories/ActorFactory.cpp b/Engine/Source/Editor/UnrealEd/Private/Factories/ActorFactory.cpp
--- a/Engine/Source/Editor/UnrealEd/Private/Factories/ActorFactory.cpp
+++ b/Engine/Source/Editor/UnrealEd/Private/Factories/ActorFactory.cpp
@@ -37,6 +37,7 @@ ActorFactory.cpp:
 #include "ActorFactories/ActorFactoryPointLight.h"
 #include "ActorFactories/ActorFactorySpotLight.h"
 #include "ActorFactories/ActorFactoryRectLight.h"
+#include "ActorFactories/ActorFactoryRuntimeVirtualTextureVolume.h"
 #include "ActorFactories/ActorFactorySkyLight.h"
 #include "ActorFactories/ActorFactorySkeletalMesh.h"
 #include "ActorFactories/ActorFactoryAnimationAsset.h"
@@ -101,6 +102,7 @@ ActorFactory.cpp:
 #include "Matinee/MatineeActor.h"
 #include "Matinee/InterpData.h"
 #include "InteractiveFoliageActor.h"
+#include "VT/RuntimeVirtualTextureVolume.h"
 
 #include "AssetRegistryModule.h"
 
@@ -1887,4 +1889,26 @@ UObject* UActorFactoryMovieScene::GetAssetFromActorInstance(AActor* Instance)
 	return nullptr;
 }
 
+/*-----------------------------------------------------------------------------
+UActorFactoryRuntimeVirtualTextureVolume
+-----------------------------------------------------------------------------*/
+UActorFactoryRuntimeVirtualTextureVolume::UActorFactoryRuntimeVirtualTextureVolume(const FObjectInitializer& ObjectInitializer)
+	: Super(ObjectInitializer)
+{
+	DisplayName = LOCTEXT("VirtualTextureVolume_DisplayName", "Virtual Texture Volume");
+	NewActorClass = ARuntimeVirtualTextureVolume::StaticClass();
+	bShowInEditorQuickMenu = 1;
+}
+
+void UActorFactoryRuntimeVirtualTextureVolume::PostSpawnActor(UObject* Asset, AActor* NewActor)
+{
+	FText ActorName = LOCTEXT("VirtualTextureVolume_DefaultActorName", "Virtual Texture");
+	NewActor->SetActorLabel(ActorName.ToString());
+
+	// Good default size to see object in editor
+	NewActor->SetActorScale3D(FVector(100.f, 100.f, 1.f));
+
+	Super::PostSpawnActor(Asset, NewActor);
+}
+
 #undef LOCTEXT_NAMESPACE
diff --git a/Engine/Source/Editor/UnrealEd/UnrealEd.Build.cs b/Engine/Source/Editor/UnrealEd/UnrealEd.Build.cs
--- a/Engine/Source/Editor/UnrealEd/UnrealEd.Build.cs
+++ b/Engine/Source/Editor/UnrealEd/UnrealEd.Build.cs
@@ -63,7 +63,8 @@ public class UnrealEd : ModuleRules
 				"ClothingSystemEditorInterface",
 				"NavigationSystem",
 				"Media",
-			}
+				"VirtualTexturingEditor",
+            }
 		);
 
 		PublicDependencyModuleNames.AddRange(
diff --git a/Engine/Source/Editor/VirtualTexturingEditor/Classes/RuntimeVirtualTextureVolumeActorFactory.h b/Engine/Source/Editor/VirtualTexturingEditor/Classes/RuntimeVirtualTextureVolumeActorFactory.h
deleted file mode 100644
diff --git a/Engine/Source/Editor/VirtualTexturingEditor/Private/RuntimeVirtualTextureVolumeActorFactory.cpp b/Engine/Source/Editor/VirtualTexturingEditor/Private/RuntimeVirtualTextureVolumeActorFactory.cpp
deleted file mode 100644
diff --git a/Engine/Source/Editor/VirtualTexturingEditor/Private/VirtualTexturingEditorModule.cpp b/Engine/Source/Editor/VirtualTexturingEditor/Private/VirtualTexturingEditorModule.cpp
--- a/Engine/Source/Editor/VirtualTexturingEditor/Private/VirtualTexturingEditorModule.cpp
+++ b/Engine/Source/Editor/VirtualTexturingEditor/Private/VirtualTexturingEditorModule.cpp
@@ -1,37 +1,51 @@
 // Copyright 1998-2019 Epic Games, Inc. All Rights Reserved.
 
-#include "AssetToolsModule.h"
+#include "VirtualTexturingEditorModule.h"
 
+#include "IPlacementModeModule.h"
 #include "PropertyEditorModule.h"
 #include "RuntimeVirtualTextureAssetTypeActions.h"
 #include "RuntimeVirtualTextureDetailsCustomization.h"
+#include "VT/RuntimeVirtualTextureVolume.h"
 
 #define LOCTEXT_NAMESPACE "VirtualTexturingEditorModule"
 
-class FVirtualTexturingEditorModule
-	: public IModuleInterface
+IMPLEMENT_MODULE(FVirtualTexturingEditorModule, VirtualTexturingEditor);
+
+void FVirtualTexturingEditorModule::StartupModule()
 {
-public:
-	virtual void StartupModule() override
-	{
-		IAssetTools& AssetTools = FModuleManager::LoadModuleChecked<FAssetToolsModule>("AssetTools").Get();
-		AssetTools.RegisterAssetTypeActions(MakeShareable(new FAssetTypeActions_RuntimeVirtualTexture));
+	IAssetTools& AssetTools = FModuleManager::LoadModuleChecked<FAssetToolsModule>("AssetTools").Get();
+	AssetTools.RegisterAssetTypeActions(MakeShareable(new FAssetTypeActions_RuntimeVirtualTexture));
 	
-		FPropertyEditorModule& PropertyModule = FModuleManager::GetModuleChecked<FPropertyEditorModule>("PropertyEditor");
-		PropertyModule.RegisterCustomClassLayout("RuntimeVirtualTexture", FOnGetDetailCustomizationInstance::CreateStatic(&FRuntimeVirtualTextureDetailsCustomization::MakeInstance));
-		PropertyModule.RegisterCustomClassLayout("RuntimeVirtualTextureComponent", FOnGetDetailCustomizationInstance::CreateStatic(&FRuntimeVirtualTextureComponentDetailsCustomization::MakeInstance));
-	}
+	FPropertyEditorModule& PropertyModule = FModuleManager::GetModuleChecked<FPropertyEditorModule>("PropertyEditor");
+	PropertyModule.RegisterCustomClassLayout("RuntimeVirtualTexture", FOnGetDetailCustomizationInstance::CreateStatic(&FRuntimeVirtualTextureDetailsCustomization::MakeInstance));
+	PropertyModule.RegisterCustomClassLayout("RuntimeVirtualTextureComponent", FOnGetDetailCustomizationInstance::CreateStatic(&FRuntimeVirtualTextureComponentDetailsCustomization::MakeInstance));
+
+	IPlacementModeModule& PlacementModeModule = IPlacementModeModule::Get();
+	PlacementModeModule.OnPlacementModeCategoryRefreshed().AddRaw(this, &FVirtualTexturingEditorModule::OnPlacementModeRefresh);
+}
 
-	virtual void ShutdownModule() override
+void FVirtualTexturingEditorModule::ShutdownModule()
+{
+	if (IPlacementModeModule::IsAvailable())
 	{
+		IPlacementModeModule::Get().OnPlacementModeCategoryRefreshed().RemoveAll(this);
 	}
+}
 
-	virtual bool SupportsDynamicReloading() override
+bool FVirtualTexturingEditorModule::SupportsDynamicReloading()
+{
+	return false;
+}
+
+void FVirtualTexturingEditorModule::OnPlacementModeRefresh(FName CategoryName)
+{
+	static FName VolumeName = FName(TEXT("Volumes"));
+	if (CategoryName == VolumeName)
 	{
-		return false;
+		IPlacementModeModule& PlacementModeModule = IPlacementModeModule::Get();
+		PlacementModeModule.RegisterPlaceableItem(CategoryName, MakeShareable(new FPlaceableItem(nullptr, FAssetData(ARuntimeVirtualTextureVolume::StaticClass()))));
 	}
-};
-
-IMPLEMENT_MODULE(FVirtualTexturingEditorModule, VirtualTexturingEditor);
+}
 
 #undef LOCTEXT_NAMESPACE
diff --git a/Engine/Source/Editor/VirtualTexturingEditor/Public/VirtualTexturingEditorModule.h b/Engine/Source/Editor/VirtualTexturingEditor/Public/VirtualTexturingEditorModule.h
new file mode 100644
--- /dev/null
+++ b/Engine/Source/Editor/VirtualTexturingEditor/Public/VirtualTexturingEditorModule.h
@@ -0,0 +1,16 @@
+// Copyright 1998-2019 Epic Games, Inc. All Rights Reserved.
+
+#include "CoreMinimal.h"
+#include "Modules/ModuleInterface.h"
+
+/** Module for virtual texturing editor extensions. */
+class FVirtualTexturingEditorModule	: public IModuleInterface
+{
+public:
+	virtual void StartupModule() override;
+	virtual void ShutdownModule() override;
+	virtual bool SupportsDynamicReloading() override;
+
+private:
+	void OnPlacementModeRefresh(FName CategoryName);
+};
diff --git a/Engine/Source/Editor/VirtualTexturingEditor/VirtualTexturingEditor.Build.cs b/Engine/Source/Editor/VirtualTexturingEditor/VirtualTexturingEditor.Build.cs
--- a/Engine/Source/Editor/VirtualTexturingEditor/VirtualTexturingEditor.Build.cs
+++ b/Engine/Source/Editor/VirtualTexturingEditor/VirtualTexturingEditor.Build.cs
@@ -31,6 +31,7 @@ public class VirtualTexturingEditor : ModuleRules
 				"EditorStyle",
 				"Engine",
 				"InputCore",
+				"PlacementMode",
 				"PropertyEditor",
 				"Slate",
 				"SlateCore",
diff --git a/Engine/Source/Runtime/Launch/Private/LaunchEngineLoop.cpp b/Engine/Source/Runtime/Launch/Private/LaunchEngineLoop.cpp
--- a/Engine/Source/Runtime/Launch/Private/LaunchEngineLoop.cpp
+++ b/Engine/Source/Runtime/Launch/Private/LaunchEngineLoop.cpp
@@ -3302,7 +3302,6 @@ bool FEngineLoop::LoadStartupCoreModules()
 		FModuleManager::Get().LoadModule(TEXT("Blutility"));
 	}
 
-	FModuleManager::Get().LoadModule(TEXT("VirtualTexturingEditor"));
 #endif //(WITH_EDITOR && !(UE_BUILD_SHIPPING || UE_BUILD_TEST))
 
 #if WITH_ENGINE
