Show dialog to save/checkout modified materials after the RVT Fix Material process is applied

diff --git a/Engine/Source/Editor/VirtualTexturingEditor/Private/RuntimeVirtualTextureAssetTypeActions.cpp b/Engine/Source/Editor/VirtualTexturingEditor/Private/RuntimeVirtualTextureAssetTypeActions.cpp
--- a/Engine/Source/Editor/VirtualTexturingEditor/Private/RuntimeVirtualTextureAssetTypeActions.cpp
+++ b/Engine/Source/Editor/VirtualTexturingEditor/Private/RuntimeVirtualTextureAssetTypeActions.cpp
@@ -5,6 +5,7 @@
 #include "AssetRegistryModule.h"
 #include "ContentBrowserModule.h"
 #include "EditorSupportDelegates.h"
+#include "FileHelpers.h"
 #include "Framework/Commands/UIAction.h"
 #include "Framework/MultiBox/MultiBoxBuilder.h"
 #include "IAssetRegistry.h"
@@ -95,92 +96,105 @@ namespace
 
 	void FixMaterialUsage(URuntimeVirtualTexture* RuntimeVirtualTexture)
 	{
-		UE_LOG(LogRuntimeVirtualTextureFixMaterial, Log, TEXT("Begin fix material usage for '%s' ..."), *RuntimeVirtualTexture->GetName());
+		TArray<UPackage*> PackagesToSave;
 
-		TArray<UMaterial*> Materials;
-		TArray<UMaterialFunctionInterface*> Functions;
-		FindAllMaterials(RuntimeVirtualTexture, Materials, Functions);
+		{
+			UE_LOG(LogRuntimeVirtualTextureFixMaterial, Log, TEXT("Begin fix material usage for '%s' ..."), *RuntimeVirtualTexture->GetName());
 
-		int32 TaskCount = Materials.Num() + Functions.Num();
-		FScopedSlowTask Task(TaskCount, LOCTEXT("RuntimeVirtualTexture_FixMaterialUsageProgress", "Fixing materials for Runtime Virtual Texture usage..."));
-		Task.MakeDialog();
+			TArray<UMaterial*> Materials;
+			TArray<UMaterialFunctionInterface*> Functions;
+			FindAllMaterials(RuntimeVirtualTexture, Materials, Functions);
 
-		for (UMaterial* Material : Materials)
-		{
-			Task.EnterProgressFrame();
+			int32 TaskCount = Materials.Num() + Functions.Num();
+			FScopedSlowTask Task(TaskCount, LOCTEXT("RuntimeVirtualTexture_FixMaterialUsageProgress", "Fixing materials for Runtime Virtual Texture usage..."));
+			Task.MakeDialog();
 
-			bool bMaterialModified = false;
-			for (UMaterialExpression* Expression : Material->Expressions)
+			for (UMaterial* Material : Materials)
 			{
-				UMaterialExpressionRuntimeVirtualTextureSample* RVTSampleExpression = Cast<UMaterialExpressionRuntimeVirtualTextureSample>(Expression);
-				if (RVTSampleExpression)
+				Task.EnterProgressFrame();
+
+				bool bMaterialModified = false;
+				for (UMaterialExpression* Expression : Material->Expressions)
 				{
-					if (RuntimeVirtualTexture == RVTSampleExpression->VirtualTexture)
+					UMaterialExpressionRuntimeVirtualTextureSample* RVTSampleExpression = Cast<UMaterialExpressionRuntimeVirtualTextureSample>(Expression);
+					if (RVTSampleExpression)
 					{
-						if (RVTSampleExpression->InitVirtualTextureDependentSettings())
+						if (RuntimeVirtualTexture == RVTSampleExpression->VirtualTexture)
 						{
-							Expression->Modify();
+							if (RVTSampleExpression->InitVirtualTextureDependentSettings())
+							{
+								Expression->Modify();
 
-							FPropertyChangedEvent Event(UMaterialExpressionTextureBase::StaticClass()->FindPropertyByName(GET_MEMBER_NAME_CHECKED(UMaterialExpressionRuntimeVirtualTextureSample, MaterialType)));
-							Expression->PostEditChangeProperty(Event);
+								FPropertyChangedEvent Event(UMaterialExpressionTextureBase::StaticClass()->FindPropertyByName(GET_MEMBER_NAME_CHECKED(UMaterialExpressionRuntimeVirtualTextureSample, MaterialType)));
+								Expression->PostEditChangeProperty(Event);
 
-							bMaterialModified = true;
+								bMaterialModified = true;
+							}
 						}
 					}
 				}
-			}
 
-			if (bMaterialModified)
-			{
-				UE_LOG(LogRuntimeVirtualTextureFixMaterial, Log, TEXT("  Recompile material '%s' ..."), *Material->GetName());
+				if (bMaterialModified)
+				{
+					UE_LOG(LogRuntimeVirtualTextureFixMaterial, Log, TEXT("  Recompile material '%s' ..."), *Material->GetName());
 
-				FScopedSlowTask CompileTask(1, FText::AsCultureInvariant(Material->GetName()));
-				CompileTask.MakeDialog();
-				CompileTask.EnterProgressFrame();
+					FScopedSlowTask CompileTask(1, FText::AsCultureInvariant(Material->GetName()));
+					CompileTask.MakeDialog();
+					CompileTask.EnterProgressFrame();
 
-				UMaterialEditingLibrary::RecompileMaterial(Material);
-			}
-		}
+					UMaterialEditingLibrary::RecompileMaterial(Material);
 
-		for (UMaterialFunctionInterface *Function : Functions)
-		{
-			Task.EnterProgressFrame();
+					PackagesToSave.Add(Material->GetOutermost());
+				}
+			}
 
-			bool bFunctionModified = false;
-			const TArray<UMaterialExpression*> *Expressions = Function->GetFunctionExpressions();
-			for (UMaterialExpression *Expression : *Expressions)
+			for (UMaterialFunctionInterface *Function : Functions)
 			{
-				UMaterialExpressionRuntimeVirtualTextureSample* RVTSampleExpression = Cast<UMaterialExpressionRuntimeVirtualTextureSample>(Expression);
-				if (RVTSampleExpression)
+				Task.EnterProgressFrame();
+
+				bool bFunctionModified = false;
+				const TArray<UMaterialExpression*> *Expressions = Function->GetFunctionExpressions();
+				for (UMaterialExpression *Expression : *Expressions)
 				{
-					if (RuntimeVirtualTexture == RVTSampleExpression->VirtualTexture)
+					UMaterialExpressionRuntimeVirtualTextureSample* RVTSampleExpression = Cast<UMaterialExpressionRuntimeVirtualTextureSample>(Expression);
+					if (RVTSampleExpression)
 					{
-						if (RVTSampleExpression->InitVirtualTextureDependentSettings())
+						if (RuntimeVirtualTexture == RVTSampleExpression->VirtualTexture)
 						{
-							Expression->Modify();
+							if (RVTSampleExpression->InitVirtualTextureDependentSettings())
+							{
+								Expression->Modify();
 
-							FPropertyChangedEvent Event(UMaterialExpressionTextureBase::StaticClass()->FindPropertyByName(GET_MEMBER_NAME_CHECKED(UMaterialExpressionRuntimeVirtualTextureSample, MaterialType)));
-							Expression->PostEditChangeProperty(Event);
+								FPropertyChangedEvent Event(UMaterialExpressionTextureBase::StaticClass()->FindPropertyByName(GET_MEMBER_NAME_CHECKED(UMaterialExpressionRuntimeVirtualTextureSample, MaterialType)));
+								Expression->PostEditChangeProperty(Event);
 
-							bFunctionModified = true;
+								bFunctionModified = true;
+							}
 						}
 					}
 				}
-			}
 
-			if (bFunctionModified)
-			{
-				UE_LOG(LogRuntimeVirtualTextureFixMaterial, Log, TEXT("  Update function '%s' ..."), *Function->GetName());
+				if (bFunctionModified)
+				{
+					UE_LOG(LogRuntimeVirtualTextureFixMaterial, Log, TEXT("  Update function '%s' ..."), *Function->GetName());
+
+					FScopedSlowTask CompileTask(1, FText::AsCultureInvariant(Function->GetName()));
+					CompileTask.MakeDialog();
+					CompileTask.EnterProgressFrame();
 
-				FScopedSlowTask CompileTask(1, FText::AsCultureInvariant(Function->GetName()));
-				CompileTask.MakeDialog();
-				CompileTask.EnterProgressFrame();
+					UMaterialEditingLibrary::UpdateMaterialFunction(Function, nullptr);
 
-				UMaterialEditingLibrary::UpdateMaterialFunction(Function, nullptr);
+					PackagesToSave.Add(Function->GetOutermost());
+				}
 			}
-		}
 
-		UE_LOG(LogRuntimeVirtualTextureFixMaterial, Log, TEXT("End fix material usage for '%s' ..."), *RuntimeVirtualTexture->GetName());
+			UE_LOG(LogRuntimeVirtualTextureFixMaterial, Log, TEXT("End fix material usage for '%s' ..."), *RuntimeVirtualTexture->GetName());
+		}
+	
+		if (PackagesToSave.Num())
+		{
+			FEditorFileUtils::PromptForCheckoutAndSave(PackagesToSave, false, true);
+		}
 	}
 }
 
