Fix case where runtime virtual texture thumbnail generation crashes on package save.

diff --git a/Engine/Source/Editor/VirtualTexturingEditor/Private/RuntimeVirtualTextureThumbnailRenderer.cpp b/Engine/Source/Editor/VirtualTexturingEditor/Private/RuntimeVirtualTextureThumbnailRenderer.cpp
--- a/Engine/Source/Editor/VirtualTexturingEditor/Private/RuntimeVirtualTextureThumbnailRenderer.cpp
+++ b/Engine/Source/Editor/VirtualTexturingEditor/Private/RuntimeVirtualTextureThumbnailRenderer.cpp
@@ -56,6 +56,12 @@ bool URuntimeVirtualTextureThumbnailRenderer::CanVisualizeAsset(UObject* Object)
 
 void URuntimeVirtualTextureThumbnailRenderer::Draw(UObject* Object, int32 X, int32 Y, uint32 Width, uint32 Height, FRenderTarget* RenderTarget, FCanvas* Canvas, bool bAdditionalViewFamily)
 {
+	//todo[vt]: Handle case where a null or floating point render target is passed in. (This happens on package save.)
+	if (RenderTarget->GetRenderTargetTexture() == nullptr || RenderTarget->GetRenderTargetTexture()->GetFormat() != PF_B8G8R8A8)
+	{
+		return;
+	}
+
 	URuntimeVirtualTexture* RuntimeVirtualTexture = Cast<URuntimeVirtualTexture>(Object);
 	URuntimeVirtualTextureComponent* RuntimeVirtualTextureComponent = FindComponent(RuntimeVirtualTexture);
 	FSceneInterface* Scene = RuntimeVirtualTextureComponent != nullptr ? RuntimeVirtualTextureComponent->GetScene() : nullptr;
