Support a mixture of formats in runtime virtual texture layers when building streaming virtual textures for low mips.
This is needed for incoming world height + world normal formats.

diff --git a/Engine/Source/Editor/VirtualTexturingEditor/Private/RuntimeVirtualTextureBuild.cpp b/Engine/Source/Editor/VirtualTexturingEditor/Private/RuntimeVirtualTextureBuild.cpp
--- a/Engine/Source/Editor/VirtualTexturingEditor/Private/RuntimeVirtualTextureBuild.cpp
+++ b/Engine/Source/Editor/VirtualTexturingEditor/Private/RuntimeVirtualTextureBuild.cpp
@@ -17,11 +17,21 @@ namespace
 	class FRenderTileResources : public FRenderResource
 	{
 	public:
-		FRenderTileResources(int32 InNumLayers, int32 InTileSize, EPixelFormat InFormat)
-			: NumLayers(InNumLayers)
-			, TileSize(InTileSize)
-			, Format(InFormat)
+		FRenderTileResources(int32 InTileSize, int32 InNumTilesX, int32 InNumTilesY, int32 InNumLayers, TArrayView<EPixelFormat> const& InLayerFormats)
+			: TileSize(InTileSize)
+			, NumLayers(InNumLayers)
+			, TotalSizeBytes(0)
 		{
+			LayerFormats.SetNumZeroed(InNumLayers);
+			LayerOffsets.SetNumZeroed(InNumLayers);
+
+			for (int32 Layer = 0; Layer < NumLayers; ++Layer)
+			{
+				check(InLayerFormats[Layer] == PF_G16 || InLayerFormats[Layer] == PF_B8G8R8A8 || InLayerFormats[Layer] == PF_DXT1 || InLayerFormats[Layer] == PF_DXT5 || InLayerFormats[Layer] == PF_BC5);
+				LayerFormats[Layer] = InLayerFormats[Layer] == PF_G16 ? PF_G16 : PF_B8G8R8A8;
+				LayerOffsets[Layer] = TotalSizeBytes;
+				TotalSizeBytes += CalculateImageBytes(InTileSize, InTileSize, 0, LayerFormats[Layer]) * InNumTilesX * InNumTilesY;
+			}
 		}
 
 		//~ Begin FRenderResource Interface.
@@ -35,8 +45,8 @@ namespace
 			for (int32 Layer = 0; Layer < NumLayers; ++Layer)
 			{
 				FRHIResourceCreateInfo CreateInfo;
-				RenderTargets[Layer] = RHICmdList.CreateTexture2D(TileSize, TileSize, Format, 1, 1, TexCreate_RenderTargetable, CreateInfo);
-				StagingTextures[Layer] = RHICmdList.CreateTexture2D(TileSize, TileSize, Format, 1, 1, TexCreate_CPUReadback, CreateInfo);
+				RenderTargets[Layer] = RHICmdList.CreateTexture2D(TileSize, TileSize, LayerFormats[Layer], 1, 1, TexCreate_RenderTargetable, CreateInfo);
+				StagingTextures[Layer] = RHICmdList.CreateTexture2D(TileSize, TileSize, LayerFormats[Layer], 1, 1, TexCreate_CPUReadback, CreateInfo);
 			}
 
 			Fence = RHICmdList.CreateGPUFence(TEXT("Runtime Virtual Texture Build"));
@@ -50,14 +60,23 @@ namespace
 		}
 		//~ End FRenderResource Interface.
 
+		int32 GetNumLayers() const { return NumLayers; }
+		int32 GetTotalSizeBytes() const { return TotalSizeBytes; }
+
+		EPixelFormat GetLayerFormat(int32 Index) const { return LayerFormats[Index]; }
+		int32 GetLayerOffset(int32 Index) const { return LayerOffsets[Index]; }
+
 		FRHITexture2D* GetRenderTarget(int32 Index) const { return Index < NumLayers ? RenderTargets[Index] : nullptr; }
 		FRHITexture2D* GetStagingTexture(int32 Index) const { return Index < NumLayers ? StagingTextures[Index] : nullptr; }
 		FRHIGPUFence* GetFence() const { return Fence; }
 
 	private:
-		int32 NumLayers;
 		int32 TileSize;
-		EPixelFormat Format;
+		int32 NumLayers;
+		int32 TotalSizeBytes;
+
+		TArray<EPixelFormat> LayerFormats;
+		TArray<int32> LayerOffsets;
 
 		TArray<FTexture2DRHIRef> RenderTargets;
 		TArray<FTexture2DRHIRef> StagingTextures;
@@ -66,25 +85,26 @@ namespace
 
 	/** Templatized helper function for copying a rendered tile to the final composited image data. */
 	template<typename T>
-	void TCopyTile(T* TilePixels, int32 TileSize, T* DestPixels, int32 DestStride, int32 DestLayerStride, FIntVector const& DestPos)
+	void TCopyTile(T* TilePixels, int32 TileSize, T* DestPixels, int32 DestStride, int32 DestLayerStride, FIntPoint const& DestPos)
 	{
 		for (int32 y = 0; y < TileSize; y++)
 		{
 			memcpy(
-				DestPixels + DestLayerStride * DestPos[2] + DestStride * (DestPos[1] + y) + DestPos[0],
+				DestPixels + DestStride * (DestPos[1] + y) + DestPos[0],
 				TilePixels + TileSize * y,
 				TileSize * sizeof(T));
 		}
 	}
 
 	/** Function for copying a rendered tile to the final composited image data. Needs ERuntimeVirtualTextureMaterialType to know what type of data is being copied. */
-	void CopyTile(void* TilePixels, int32 TileSize, void* DestPixels, int32 DestStride, int32 DestLayerStride, FIntVector const& DestPos, ERuntimeVirtualTextureMaterialType MaterialType)
+	void CopyTile(void* TilePixels, int32 TileSize, void* DestPixels, int32 DestStride, int32 DestLayerStride, FIntPoint const& DestPos, EPixelFormat Format)
 	{
-		if (MaterialType == ERuntimeVirtualTextureMaterialType::WorldHeight)
+		check(Format == PF_G16 || Format == PF_B8G8R8A8);
+		if (Format == PF_G16)
 		{
 			TCopyTile((uint16*)TilePixels, TileSize, (uint16*)DestPixels, DestStride, DestLayerStride, DestPos);
 		}
-		else
+		else if (Format == PF_B8G8R8A8)
 		{
 			TCopyTile((FColor*)TilePixels, TileSize, (FColor*)DestPixels, DestStride, DestLayerStride, DestPos);
 		}
@@ -144,8 +164,11 @@ namespace RuntimeVirtualTexture
 		const int32 NumLayers = RuntimeVirtualTexture->GetLayerCount();
 
 		const ERuntimeVirtualTextureMaterialType MaterialType = RuntimeVirtualTexture->GetMaterialType();
-		const EPixelFormat RenderTargetFormat = (MaterialType == ERuntimeVirtualTextureMaterialType::WorldHeight) ? PF_G16 : PF_B8G8R8A8;
-		const int32 BytesPerPixel = (MaterialType == ERuntimeVirtualTextureMaterialType::WorldHeight) ? 2 : 4;
+		TArray<EPixelFormat, TInlineAllocator<4>> LayerFormats;
+		for (int32 Layer = 0; Layer < NumLayers; ++Layer)
+		{
+			LayerFormats.Add(RuntimeVirtualTexture->GetLayerFormat(Layer));
+		}
 
 		// Spin up slow task UI
 		const float TaskWorkRender = NumTilesX * NumTilesY;
@@ -154,12 +177,12 @@ namespace RuntimeVirtualTexture
 		Task.MakeDialog(true);
 
 		// Allocate render targets for rendering out the runtime virtual texture tiles
-		FRenderTileResources RenderTileResources(NumLayers, TileSize, RenderTargetFormat);
+		FRenderTileResources RenderTileResources(TileSize, NumTilesX, NumTilesY, NumLayers, LayerFormats);
 		BeginInitResource(&RenderTileResources);
 
 		// Final pixels will contain image data for each virtual texture layer in order
 		TArray64<uint8> FinalPixels;
-		FinalPixels.InsertUninitialized(0, ImageSizeX * ImageSizeY * NumLayers * BytesPerPixel);
+		FinalPixels.SetNumUninitialized(RenderTileResources.GetTotalSizeBytes());
 
 		// Iterate over all tiles and render/store each one to the final image
 		for (int32 TileY = 0; TileY < NumTilesY && !Task.ShouldCancel(); TileY++)
@@ -241,8 +264,11 @@ namespace RuntimeVirtualTexture
 						check(TilePixels != nullptr);
 						check(OutWidth == TileSize && OutHeight == TileSize);
 
-						const FIntVector DestPos(TileX * TileSize, TileY * TileSize, Layer);
-						CopyTile(TilePixels, TileSize, FinalPixels.GetData(), ImageSizeX, ImageSizeX * ImageSizeY, DestPos, MaterialType);
+						const int32 LayerOffset = RenderTileResources.GetLayerOffset(Layer);
+						const EPixelFormat LayerFormat = RenderTileResources.GetLayerFormat(Layer);
+						const FIntPoint DestPos(TileX * TileSize, TileY * TileSize);
+
+						CopyTile(TilePixels, TileSize, FinalPixels.GetData() + LayerOffset, ImageSizeX, ImageSizeX * ImageSizeY, DestPos, LayerFormat);
 
 						RHICmdList.UnmapStagingSurface(RenderTileResources.GetStagingTexture(Layer));
 					}
