Add UVirtualTextureBuilder asset type for building virtual textures.
Initial use case is to move the runtime virtual texture low mips building out of specific runtime virtual texture code.

diff --git a/Engine/Source/Editor/VirtualTexturingEditor/Classes/VirtualTextureBuilderAssetTypeActions.h b/Engine/Source/Editor/VirtualTexturingEditor/Classes/VirtualTextureBuilderAssetTypeActions.h
new file mode 100644
--- /dev/null
+++ b/Engine/Source/Editor/VirtualTexturingEditor/Classes/VirtualTextureBuilderAssetTypeActions.h
@@ -0,0 +1,23 @@
+// Copyright Epic Games, Inc. All Rights Reserved.
+
+#pragma once
+
+#include "CoreMinimal.h"
+#include "AssetTypeActions_Base.h"
+
+class UVirtualTextureBuilder;
+
+/** Asset actions setup for URuntimeVirtualTexture */
+class FAssetTypeActions_VirtualTextureBuilder : public FAssetTypeActions_Base
+{
+public:
+	FAssetTypeActions_VirtualTextureBuilder() {}
+
+protected:
+	//~ Begin IAssetTypeActions Interface.
+	virtual UClass* GetSupportedClass() const override;
+	virtual FText GetName() const override;
+	virtual FColor GetTypeColor() const override;
+	virtual uint32 GetCategories() override;
+	//~ End IAssetTypeActions Interface.
+};
diff --git a/Engine/Source/Editor/VirtualTexturingEditor/Classes/VirtualTextureBuilderFactory.h b/Engine/Source/Editor/VirtualTexturingEditor/Classes/VirtualTextureBuilderFactory.h
new file mode 100644
--- /dev/null
+++ b/Engine/Source/Editor/VirtualTexturingEditor/Classes/VirtualTextureBuilderFactory.h
@@ -0,0 +1,19 @@
+// Copyright Epic Games, Inc. All Rights Reserved.
+
+#pragma once
+
+#include "CoreMinimal.h"
+#include "UObject/ObjectMacros.h"
+#include "Factories/Factory.h"
+#include "VirtualTextureBuilderFactory.generated.h"
+
+/** Factory for UVirtualTextureBuilder */
+UCLASS(hidecategories = (Object))
+class UVirtualTextureBuilderFactory : public UFactory
+{
+	GENERATED_UCLASS_BODY()
+
+	//~ Begin UFactory Interface
+	virtual UObject* FactoryCreateNew(UClass* Class, UObject* InParent, FName Name, EObjectFlags Flags, UObject* Context, FFeedbackContext* Warn) override;
+	//~ Begin UFactory Interface	
+};
diff --git a/Engine/Source/Editor/VirtualTexturingEditor/Private/VirtualTextureBuilderAssetTypeActions.cpp b/Engine/Source/Editor/VirtualTexturingEditor/Private/VirtualTextureBuilderAssetTypeActions.cpp
new file mode 100644
--- /dev/null
+++ b/Engine/Source/Editor/VirtualTexturingEditor/Private/VirtualTextureBuilderAssetTypeActions.cpp
@@ -0,0 +1,29 @@
+// Copyright Epic Games, Inc. All Rights Reserved.
+
+#include "VirtualTextureBuilderAssetTypeActions.h"
+
+#include "VT/VirtualTextureBuilder.h"
+
+#define LOCTEXT_NAMESPACE "VirtualTexturingEditorModule"
+
+UClass* FAssetTypeActions_VirtualTextureBuilder::GetSupportedClass() const
+{
+	return UVirtualTextureBuilder::StaticClass();
+}
+
+FText FAssetTypeActions_VirtualTextureBuilder::GetName() const
+{
+	return LOCTEXT("AssetTypeActions_VirtualTextureBuilder", "Virtual Texture Builder"); 
+}
+
+FColor FAssetTypeActions_VirtualTextureBuilder::GetTypeColor() const 
+{
+	return FColor(128, 128, 128); 
+}
+
+uint32 FAssetTypeActions_VirtualTextureBuilder::GetCategories() 
+{
+	return EAssetTypeCategories::MaterialsAndTextures; 
+}
+
+#undef LOCTEXT_NAMESPACE
diff --git a/Engine/Source/Editor/VirtualTexturingEditor/Private/VirtualTextureBuilderFactory.cpp b/Engine/Source/Editor/VirtualTexturingEditor/Private/VirtualTextureBuilderFactory.cpp
new file mode 100644
--- /dev/null
+++ b/Engine/Source/Editor/VirtualTexturingEditor/Private/VirtualTextureBuilderFactory.cpp
@@ -0,0 +1,21 @@
+// Copyright Epic Games, Inc. All Rights Reserved.
+
+#include "VirtualTextureBuilderFactory.h"
+
+#include "AssetTypeCategories.h"
+#include "VT/VirtualTextureBuilder.h"
+
+UVirtualTextureBuilderFactory::UVirtualTextureBuilderFactory(const FObjectInitializer& ObjectInitializer)
+	: Super(ObjectInitializer)
+{
+	SupportedClass = UVirtualTextureBuilder::StaticClass();
+	
+	bCreateNew = true;
+	bEditAfterNew = false;
+	bEditorImport = false;
+}
+
+UObject* UVirtualTextureBuilderFactory::FactoryCreateNew(UClass* Class, UObject* InParent, FName Name, EObjectFlags Flags, UObject* Context, FFeedbackContext* Warn)
+{
+	return NewObject<UVirtualTextureBuilder>(InParent, Class, Name, Flags);
+}
diff --git a/Engine/Source/Editor/VirtualTexturingEditor/Private/VirtualTextureBuilderThumbnailRenderer.cpp b/Engine/Source/Editor/VirtualTexturingEditor/Private/VirtualTextureBuilderThumbnailRenderer.cpp
new file mode 100644
--- /dev/null
+++ b/Engine/Source/Editor/VirtualTexturingEditor/Private/VirtualTextureBuilderThumbnailRenderer.cpp
@@ -0,0 +1,27 @@
+// Copyright Epic Games, Inc. All Rights Reserved.
+
+#include "VirtualTextureBuilderThumbnailRenderer.h"
+#include "VT/VirtualTexture.h"
+#include "VT/VirtualTextureBuilder.h"
+
+UVirtualTextureBuilderThumbnailRenderer::UVirtualTextureBuilderThumbnailRenderer(const FObjectInitializer& ObjectInitializer)
+	: Super(ObjectInitializer)
+{
+}
+
+bool UVirtualTextureBuilderThumbnailRenderer::CanVisualizeAsset(UObject* Object)
+{
+	UVirtualTextureBuilder* VirtualTextureBuilder = Cast<UVirtualTextureBuilder>(Object);
+	UVirtualTexture2D* Texture = VirtualTextureBuilder != nullptr ? VirtualTextureBuilder->Texture : nullptr;
+	return Texture != nullptr ? UTextureThumbnailRenderer::CanVisualizeAsset(Texture) : false;
+}
+
+void UVirtualTextureBuilderThumbnailRenderer::Draw(UObject* Object, int32 X, int32 Y, uint32 Width, uint32 Height, FRenderTarget* RenderTarget, FCanvas* Canvas, bool bAdditionalViewFamily)
+{
+	UVirtualTextureBuilder* VirtualTextureBuilder = Cast<UVirtualTextureBuilder>(Object);
+	UVirtualTexture2D* Texture = VirtualTextureBuilder != nullptr ? VirtualTextureBuilder->Texture : nullptr;
+	if (Texture != nullptr)
+	{
+		UTextureThumbnailRenderer::Draw(Texture, X, Y, Width, Height, RenderTarget, Canvas, bAdditionalViewFamily);
+	}
+}
diff --git a/Engine/Source/Editor/VirtualTexturingEditor/Private/VirtualTextureBuilderThumbnailRenderer.h b/Engine/Source/Editor/VirtualTexturingEditor/Private/VirtualTextureBuilderThumbnailRenderer.h
new file mode 100644
--- /dev/null
+++ b/Engine/Source/Editor/VirtualTexturingEditor/Private/VirtualTextureBuilderThumbnailRenderer.h
@@ -0,0 +1,18 @@
+// Copyright Epic Games, Inc. All Rights Reserved.
+
+#pragma once
+
+#include "CoreMinimal.h"
+#include "ThumbnailRendering/TextureThumbnailRenderer.h"
+#include "VirtualTextureBuilderThumbnailRenderer.generated.h"
+
+UCLASS(MinimalAPI)
+class UVirtualTextureBuilderThumbnailRenderer : public UTextureThumbnailRenderer
+{
+	GENERATED_UCLASS_BODY()
+
+	//~ Begin UThumbnailRenderer Interface.
+	virtual bool CanVisualizeAsset(UObject* Object) override;
+	virtual void Draw(UObject* Object, int32 X, int32 Y, uint32 Width, uint32 Height, FRenderTarget*, FCanvas* Canvas, bool bAdditionalViewFamily) override;
+	//~ EndUThumbnailRenderer Interface.
+};
diff --git a/Engine/Source/Editor/VirtualTexturingEditor/Private/VirtualTexturingEditorModule.cpp b/Engine/Source/Editor/VirtualTexturingEditor/Private/VirtualTexturingEditorModule.cpp
--- a/Engine/Source/Editor/VirtualTexturingEditor/Private/VirtualTexturingEditorModule.cpp
+++ b/Engine/Source/Editor/VirtualTexturingEditor/Private/VirtualTexturingEditorModule.cpp
@@ -9,8 +9,11 @@
 #include "RuntimeVirtualTextureDetailsCustomization.h"
 #include "RuntimeVirtualTextureThumbnailRenderer.h"
 #include "ThumbnailRendering/ThumbnailManager.h"
+#include "VirtualTextureBuilderAssetTypeActions.h"
+#include "VirtualTextureBuilderThumbnailRenderer.h"
 #include "VT/RuntimeVirtualTexture.h"
 #include "VT/RuntimeVirtualTextureVolume.h"
+#include "VT/VirtualTextureBuilder.h"
 
 #define LOCTEXT_NAMESPACE "VirtualTexturingEditorModule"
 
@@ -39,7 +42,8 @@ void FVirtualTexturingEditorModule::StartupModule()
 {
 	IAssetTools& AssetTools = FModuleManager::LoadModuleChecked<FAssetToolsModule>("AssetTools").Get();
 	AssetTools.RegisterAssetTypeActions(MakeShareable(new FAssetTypeActions_RuntimeVirtualTexture));
-	
+	AssetTools.RegisterAssetTypeActions(MakeShareable(new FAssetTypeActions_VirtualTextureBuilder));
+
 	FPropertyEditorModule& PropertyModule = FModuleManager::GetModuleChecked<FPropertyEditorModule>("PropertyEditor");
 	PropertyModule.RegisterCustomClassLayout("RuntimeVirtualTexture", FOnGetDetailCustomizationInstance::CreateStatic(&FRuntimeVirtualTextureDetailsCustomization::MakeInstance));
 	PropertyModule.RegisterCustomClassLayout("RuntimeVirtualTextureComponent", FOnGetDetailCustomizationInstance::CreateStatic(&FRuntimeVirtualTextureComponentDetailsCustomization::MakeInstance));
@@ -48,6 +52,7 @@ void FVirtualTexturingEditorModule::StartupModule()
 	PlacementModeModule.OnPlacementModeCategoryRefreshed().AddRaw(this, &FVirtualTexturingEditorModule::OnPlacementModeRefresh);
 
 	UThumbnailManager::Get().RegisterCustomRenderer(URuntimeVirtualTexture::StaticClass(), URuntimeVirtualTextureThumbnailRenderer::StaticClass());
+	UThumbnailManager::Get().RegisterCustomRenderer(UVirtualTextureBuilder::StaticClass(), UVirtualTextureBuilderThumbnailRenderer::StaticClass());
 }
 
 void FVirtualTexturingEditorModule::ShutdownModule()
diff --git a/Engine/Source/Runtime/Engine/Classes/VT/VirtualTexture.h b/Engine/Source/Runtime/Engine/Classes/VT/VirtualTexture.h
--- a/Engine/Source/Runtime/Engine/Classes/VT/VirtualTexture.h
+++ b/Engine/Source/Runtime/Engine/Classes/VT/VirtualTexture.h
@@ -3,6 +3,8 @@
 #pragma once
 
 #include "CoreMinimal.h"
+#include "Engine/Texture2D.h"
+#include "VT/VirtualTextureBuildSettings.h"
 #include "VirtualTexture.generated.h"
 
 /** Deprecated class */
@@ -19,3 +21,27 @@ class ENGINE_API ULightMapVirtualTexture : public UVirtualTexture
 {
 	GENERATED_UCLASS_BODY()
 };
+
+/**
+ * Virtual Texture with locally configurable build settings.
+ * A raw UTexture2D can also represent a Virtual Texture but uses the one and only per-project build settings.
+ */
+UCLASS(ClassGroup = Rendering)
+class UVirtualTexture2D : public UTexture2D
+{
+	GENERATED_UCLASS_BODY()
+
+		UPROPERTY()
+		FVirtualTextureBuildSettings Settings;
+
+	UPROPERTY()
+		bool bSinglePhysicalSpace;
+
+	//~ Begin UTexture Interface.
+	virtual void GetVirtualTextureBuildSettings(FVirtualTextureBuildSettings& OutSettings) const override { OutSettings = Settings; }
+	virtual bool IsVirtualTexturedWithSinglePhysicalSpace() const override { return bSinglePhysicalSpace; }
+#if WITH_EDITOR
+	virtual void BeginCacheForCookedPlatformData(const ITargetPlatform* TargetPlatform) override;
+#endif
+	//~ End UTexture Interface.
+};
diff --git a/Engine/Source/Runtime/Engine/Classes/VT/VirtualTextureBuilder.h b/Engine/Source/Runtime/Engine/Classes/VT/VirtualTextureBuilder.h
new file mode 100644
--- /dev/null
+++ b/Engine/Source/Runtime/Engine/Classes/VT/VirtualTextureBuilder.h
@@ -0,0 +1,68 @@
+// Copyright Epic Games, Inc. All Rights Reserved.
+
+#pragma once
+
+#include "CoreMinimal.h"
+#include "Engine/Texture.h"
+#include "VirtualTextureBuilder.generated.h"
+
+#if WITH_EDITOR
+
+/** Description object used to build the contents of a UVirtualTextureBuilder. */
+USTRUCT(BlueprintType)
+struct FVirtualTextureBuildDesc
+{
+	GENERATED_USTRUCT_BODY()
+
+		uint32 BuildHash = 0;
+
+	int32 LayerCount = 0;
+	TArray<ETextureSourceFormat, TInlineAllocator<4>> LayerFormats;
+	TArray<FTextureFormatSettings, TInlineAllocator<4>> LayerFormatSettings;
+
+	int32 TileSize = 0;
+	int32 TileBorderSize = 0;
+
+	bool bCrunchCompressed = false;
+	bool bSinglePhysicalSpace = false;
+
+	int32 NumMips = 0;
+
+	uint32 InSizeX = 0;
+	uint32 InSizeY = 0;
+	uint8 const* InData = nullptr;
+};
+
+#endif
+
+/**
+ * Container for a UVirtualTexture2D that can be built from a FVirtualTextureBuildDesc description.
+ * This has a simple BuildTexture() interface but we may want to extend in the future to support partial builds
+ * or other more blueprint driven approaches for data generation.
+ */
+UCLASS(ClassGroup = Rendering, BlueprintType)
+class ENGINE_API UVirtualTextureBuilder : public UObject
+{
+public:
+	GENERATED_UCLASS_BODY()
+		~UVirtualTextureBuilder();
+
+	/** The UTexture object. */
+	UPROPERTY(VisibleAnywhere, BlueprintReadOnly, Category = Texture)
+		class UVirtualTexture2D* Texture;
+
+	/** Some client defined hash of that defines how the Texture was built. */
+	UPROPERTY()
+		uint32 BuildHash;
+
+#if WITH_EDITOR
+	UFUNCTION(BlueprintCallable, Category = Texture)
+		/** Creates a new UVirtualTexture2D and stores it in the contained Texture. */
+		void BuildTexture(FVirtualTextureBuildDesc const& BuildDesc);
+#endif
+
+protected:
+	//~ Begin UObject Interface
+	virtual void Serialize(FArchive& Ar) override;
+	//~ End UObject Interface
+};
diff --git a/Engine/Source/Runtime/Engine/Private/VT/VirtualTexture.cpp b/Engine/Source/Runtime/Engine/Private/VT/VirtualTexture.cpp
--- a/Engine/Source/Runtime/Engine/Private/VT/VirtualTexture.cpp
+++ b/Engine/Source/Runtime/Engine/Private/VT/VirtualTexture.cpp
@@ -172,3 +172,27 @@ ULightMapVirtualTexture::ULightMapVirtualTexture(const FObjectInitializer& Objec
 	: Super(ObjectInitializer) 
 {
 }
+
+
+UVirtualTexture2D::UVirtualTexture2D(const FObjectInitializer& ObjectInitializer)
+	: UTexture2D(ObjectInitializer)
+	, bSinglePhysicalSpace(false)
+{
+}
+
+#if WITH_EDITOR
+
+void UVirtualTexture2D::BeginCacheForCookedPlatformData(const ITargetPlatform* TargetPlatform)
+{
+	// Even though we skip the cook of this object for non VT platforms in URuntimeVirtualTexture::Serialize()
+	// we still load the object at cook time and kick off the DDC build. This will trigger an error in the texture DDC code.
+	// Either we need to make the DDC code more robust for non VT platforms or we can skip the process here...
+	if (!UseVirtualTexturing(GMaxRHIFeatureLevel, TargetPlatform))
+	{
+		return;
+	}
+
+	Super::BeginCacheForCookedPlatformData(TargetPlatform);
+}
+
+#endif
diff --git a/Engine/Source/Runtime/Engine/Private/VT/VirtualTextureBuilder.cpp b/Engine/Source/Runtime/Engine/Private/VT/VirtualTextureBuilder.cpp
new file mode 100644
--- /dev/null
+++ b/Engine/Source/Runtime/Engine/Private/VT/VirtualTextureBuilder.cpp
@@ -0,0 +1,58 @@
+// Copyright Epic Games, Inc. All Rights Reserved.
+
+#include "VT/VirtualTextureBuilder.h"
+
+#include "RenderUtils.h"
+#include "VT/VirtualTexture.h"
+
+UVirtualTextureBuilder::UVirtualTextureBuilder(const FObjectInitializer& ObjectInitializer)
+	: UObject(ObjectInitializer)
+{
+}
+
+UVirtualTextureBuilder::~UVirtualTextureBuilder()
+{
+}
+
+void UVirtualTextureBuilder::Serialize(FArchive& Ar)
+{
+	if (Ar.IsCooking() && Ar.IsSaving() && !UseVirtualTexturing(GMaxRHIFeatureLevel, Ar.CookingTarget()))
+	{
+		// Clear Texture during cook for platforms that don't support virtual texturing
+		UVirtualTexture2D* TextureBackup = Texture;
+		Texture = nullptr;
+		Super::Serialize(Ar);
+		Texture = TextureBackup;
+	}
+	else
+	{
+		Super::Serialize(Ar);
+	}
+}
+
+#if WITH_EDITOR
+
+void UVirtualTextureBuilder::BuildTexture(FVirtualTextureBuildDesc const& BuildDesc)
+{
+	BuildHash = BuildDesc.BuildHash;
+
+	Texture = NewObject<UVirtualTexture2D>(this, TEXT("Texture"));
+	Texture->VirtualTextureStreaming = true;
+
+	Texture->Settings.Init();
+	Texture->Settings.TileSize = BuildDesc.TileSize;
+	Texture->Settings.TileBorderSize = BuildDesc.TileBorderSize;
+	Texture->Settings.bEnableCompressCrunch = BuildDesc.bCrunchCompressed;
+
+	Texture->bSinglePhysicalSpace = BuildDesc.bSinglePhysicalSpace;
+
+	for (int32 Layer = 0; Layer < BuildDesc.LayerCount; Layer++)
+	{
+		Texture->SetLayerFormatSettings(Layer, BuildDesc.LayerFormatSettings[Layer]);
+	}
+
+	Texture->Source.InitLayered(BuildDesc.InSizeX, BuildDesc.InSizeY, 1, BuildDesc.LayerCount, 1, BuildDesc.LayerFormats.GetData(), BuildDesc.InData);
+	Texture->PostEditChange();
+}
+
+#endif
