Move bounds copy functionality into UBoundsCopyComponent for reuse

diff --git a/Engine/Source/Editor/DetailCustomizations/Private/BoundsCopyComponentDetails.cpp b/Engine/Source/Editor/DetailCustomizations/Private/BoundsCopyComponentDetails.cpp
new file mode 100644
--- /dev/null
+++ b/Engine/Source/Editor/DetailCustomizations/Private/BoundsCopyComponentDetails.cpp
@@ -0,0 +1,113 @@
+// Copyright 1998-2019 Epic Games, Inc. All Rights Reserved.
+
+#include "BoundsCopyComponentDetails.h"
+
+#include "Components/BoundsCopyComponent.h"
+#include "DetailCategoryBuilder.h"
+#include "DetailLayoutBuilder.h"
+#include "DetailWidgetRow.h"
+#include "ScopedTransaction.h"
+#include "Widgets/Input/SButton.h"
+#include "Widgets/Layout/SBox.h"
+#include "Widgets/Layout/SWrapBox.h"
+
+#define LOCTEXT_NAMESPACE "BoundsCopyComponentDetails"
+
+FBoundsCopyComponentDetailsCustomization::FBoundsCopyComponentDetailsCustomization()
+{
+}
+
+TSharedRef<IDetailCustomization> FBoundsCopyComponentDetailsCustomization::MakeInstance()
+{
+	return MakeShareable(new FBoundsCopyComponentDetailsCustomization);
+}
+
+void FBoundsCopyComponentDetailsCustomization::CustomizeDetails(IDetailLayoutBuilder& DetailBuilder)
+{
+	// Get and store the linked UBoundsCopyComponent
+	TArray<TWeakObjectPtr<UObject>> ObjectsBeingCustomized;
+	DetailBuilder.GetObjectsBeingCustomized(ObjectsBeingCustomized);
+	if (ObjectsBeingCustomized.Num() > 1)
+	{
+		return;
+	}
+	BoundsCopyComponent = Cast<UBoundsCopyComponent>(ObjectsBeingCustomized[0].Get());
+	if (BoundsCopyComponent == nullptr)
+	{
+		return;
+	}
+
+	// Only reason for having any of the logic here is that CallInEditor doesn't seem to work to add buttons for Copy functions.
+	IDetailCategoryBuilder& BoundsCategory = DetailBuilder.EditCategory("TransformFromBounds", FText::GetEmpty(), ECategoryPriority::Important);
+
+	// Hide and re-add BoundsSourceActor property otherwise we lose the ordering of this property first.
+	TSharedPtr<IPropertyHandle> SourceActorValue = DetailBuilder.GetProperty("BoundsSourceActor");
+	DetailBuilder.HideProperty(SourceActorValue);
+
+	BoundsCategory.AddCustomRow(SourceActorValue->GetPropertyDisplayName())
+	.NameContent()
+	[
+		SourceActorValue->CreatePropertyNameWidget()
+	]
+	.ValueContent()
+	[
+		SourceActorValue->CreatePropertyValueWidget()
+	];
+
+	// Add Copy buttons.
+	BoundsCategory
+	.AddCustomRow(LOCTEXT("Button_CopyRotation", "Copy Rotation"))
+	.NameContent()
+	[
+		SNew(STextBlock)
+		.Font(IDetailLayoutBuilder::GetDetailFont())
+		.Text(LOCTEXT("Button_CopyRotation", "Copy Rotation"))
+	]
+	.ValueContent()
+	.MaxDesiredWidth(125.f)
+	[
+		SNew(SButton)
+		.VAlign(VAlign_Center)
+		.HAlign(HAlign_Center)
+		.ContentPadding(2)
+		.Text(LOCTEXT("Button_Copy", "Copy"))
+		.ToolTipText(LOCTEXT("Button_CopyRotation_Tooltip", "Set the virtual texture rotation to match the source actor"))
+		.OnClicked(this, &FBoundsCopyComponentDetailsCustomization::SetRotation)
+	];
+
+	BoundsCategory
+	.AddCustomRow(LOCTEXT("Button_CopyBounds", "Copy Bounds"))
+	.NameContent()
+	[
+		SNew(STextBlock)
+		.Font(IDetailLayoutBuilder::GetDetailFont())
+		.Text(LOCTEXT("Button_CopyBounds", "Copy Bounds"))
+	]
+	.ValueContent()
+	.MaxDesiredWidth(125.f)
+	[
+		SNew(SButton)
+		.VAlign(VAlign_Center)
+		.HAlign(HAlign_Center)
+		.ContentPadding(2)
+		.Text(LOCTEXT("Button_Copy", "Copy"))
+		.ToolTipText(LOCTEXT("Button_CopyBounds_Tooltip", "Set the virtual texture transform so that it includes the full bounds of the source actor"))
+		.OnClicked(this, &FBoundsCopyComponentDetailsCustomization::SetTransformToBounds)
+	];
+}
+
+FReply FBoundsCopyComponentDetailsCustomization::SetRotation()
+{
+	FScopedTransaction BakeTransaction(LOCTEXT("Transaction_CopyRotation", "Copy Rotation"));
+	BoundsCopyComponent->SetRotation();
+	return FReply::Handled();
+}
+
+FReply FBoundsCopyComponentDetailsCustomization::SetTransformToBounds()
+{
+	FScopedTransaction BakeTransaction(LOCTEXT("Transaction_CopyBounds", "Copy Bounds"));
+	BoundsCopyComponent->SetTransformToBounds();
+	return FReply::Handled();
+}
+
+#undef LOCTEXT_NAMESPACE
diff --git a/Engine/Source/Editor/DetailCustomizations/Private/BoundsCopyComponentDetails.h b/Engine/Source/Editor/DetailCustomizations/Private/BoundsCopyComponentDetails.h
new file mode 100644
--- /dev/null
+++ b/Engine/Source/Editor/DetailCustomizations/Private/BoundsCopyComponentDetails.h
@@ -0,0 +1,30 @@
+// Copyright 1998-2019 Epic Games, Inc. All Rights Reserved.
+
+#pragma once
+
+#include "CoreMinimal.h"
+#include "IDetailCustomization.h"
+#include "Widgets/Text/STextBlock.h"
+#include "Input/Reply.h"
+
+/** UI customization for UBoundsCopyComponent */
+class FBoundsCopyComponentDetailsCustomization : public IDetailCustomization
+{
+public:
+	static TSharedRef<IDetailCustomization> MakeInstance();
+
+protected:
+	FBoundsCopyComponentDetailsCustomization();
+
+	/** Callback for Copy Rotation button */
+	FReply SetRotation();
+	/** Callback for Copy Bounds button */
+	FReply SetTransformToBounds();
+
+	//~ Begin IDetailCustomization Interface.
+	virtual void CustomizeDetails(IDetailLayoutBuilder& DetailBuilder) override;
+	//~ End IDetailCustomization Interface.
+
+private:
+	class UBoundsCopyComponent* BoundsCopyComponent;
+};
diff --git a/Engine/Source/Editor/DetailCustomizations/Private/DetailCustomizations.cpp b/Engine/Source/Editor/DetailCustomizations/Private/DetailCustomizations.cpp
--- a/Engine/Source/Editor/DetailCustomizations/Private/DetailCustomizations.cpp
+++ b/Engine/Source/Editor/DetailCustomizations/Private/DetailCustomizations.cpp
@@ -139,8 +139,9 @@
 #include "SkeletonDetails.h"
 #include "MaterialShadingModelCustomization.h"
 #include "DebugCameraControllerSettingsCustomization.h"
-IMPLEMENT_MODULE( FDetailCustomizationsModule, DetailCustomizations );
+#include "BoundsCopyComponentDetails.h"
 
+IMPLEMENT_MODULE( FDetailCustomizationsModule, DetailCustomizations );
 
 void FDetailCustomizationsModule::StartupModule()
 {
@@ -388,6 +389,8 @@ void FDetailCustomizationsModule::RegisterObjectCustomizations()
 
 	RegisterCustomClassLayout("Landscape", FOnGetDetailCustomizationInstance::CreateStatic(&FLandscapeUIDetails::MakeInstance));
 	RegisterCustomClassLayout("LandscapeProxy", FOnGetDetailCustomizationInstance::CreateStatic(&FLandscapeProxyUIDetails::MakeInstance));
+
+	RegisterCustomClassLayout("BoundsCopyComponent", FOnGetDetailCustomizationInstance::CreateStatic(&FBoundsCopyComponentDetailsCustomization::MakeInstance));
 }
 
 
diff --git a/Engine/Source/Editor/VirtualTexturingEditor/Classes/RuntimeVirtualTextureDetailsCustomization.h b/Engine/Source/Editor/VirtualTexturingEditor/Classes/RuntimeVirtualTextureDetailsCustomization.h
--- a/Engine/Source/Editor/VirtualTexturingEditor/Classes/RuntimeVirtualTextureDetailsCustomization.h
+++ b/Engine/Source/Editor/VirtualTexturingEditor/Classes/RuntimeVirtualTextureDetailsCustomization.h
@@ -45,10 +45,6 @@ public:
 protected:
 	FRuntimeVirtualTextureComponentDetailsCustomization();
 
-	/** Callback for Copy Rotation button */
-	FReply SetRotation();
-	/** Callback for Copy Bounds button */
-	FReply SetTransformToBounds();
 	/** Callback for Build button */
 	FReply BuildStreamedMips();
 	/** Callback for Build Debug button */
diff --git a/Engine/Source/Editor/VirtualTexturingEditor/Private/RuntimeVirtualTextureDetailsCustomization.cpp b/Engine/Source/Editor/VirtualTexturingEditor/Private/RuntimeVirtualTextureDetailsCustomization.cpp
--- a/Engine/Source/Editor/VirtualTexturingEditor/Private/RuntimeVirtualTextureDetailsCustomization.cpp
+++ b/Engine/Source/Editor/VirtualTexturingEditor/Private/RuntimeVirtualTextureDetailsCustomization.cpp
@@ -171,7 +171,7 @@ TSharedRef<IDetailCustomization> FRuntimeVirtualTextureComponentDetailsCustomiza
 
 void FRuntimeVirtualTextureComponentDetailsCustomization::CustomizeDetails(IDetailLayoutBuilder& DetailBuilder)
 {
-	// Get and store the linked ARuntimeRuntimeVirtualTextureComponent
+	// Get and store the linked URuntimeVirtualTextureComponent
 	TArray<TWeakObjectPtr<UObject>> ObjectsBeingCustomized;
 	DetailBuilder.GetObjectsBeingCustomized(ObjectsBeingCustomized);
 	if (ObjectsBeingCustomized.Num() > 1)
@@ -184,126 +184,49 @@ void FRuntimeVirtualTextureComponentDetailsCustomization::CustomizeDetails(IDeta
 		return;
 	}
 
-	// Use existing property to add bounds copy buttons
-	TSharedPtr<IPropertyHandle> SourceActorValue = DetailBuilder.GetProperty("BoundsSourceActor");
-	DetailBuilder.HideProperty(SourceActorValue);
-
-	IDetailCategoryBuilder& BoundsCategory = DetailBuilder.EditCategory("TransformFromBounds", FText::GetEmpty(), ECategoryPriority::Important);
-	BoundsCategory.AddCustomRow(SourceActorValue->GetPropertyDisplayName())
+	IDetailCategoryBuilder& VirtualTextureCategory = DetailBuilder.EditCategory("VirtualTextureBuild", FText::GetEmpty());
+	
+	VirtualTextureCategory
+	.AddCustomRow(LOCTEXT("Button_BuildStreamingMips", "Build Streaming Mips"), true)
 	.NameContent()
 	[
-		SourceActorValue->CreatePropertyNameWidget()
+		SNew(STextBlock)
+		.Font(IDetailLayoutBuilder::GetDetailFont())
+		.Text(LOCTEXT("Button_BuildStreamingMips", "Build Streaming Mips"))
 	]
 	.ValueContent()
-	.MaxDesiredWidth(TOptional<float>())
+	.MaxDesiredWidth(125.f)
 	[
-		SNew(SHorizontalBox)
-		
-		+ SHorizontalBox::Slot()
-		[
-			SourceActorValue->CreatePropertyValueWidget()
-		]
-
-		+ SHorizontalBox::Slot()
-		[
-			SNew(SWrapBox)
-			
-			+ SWrapBox::Slot()
-			.Padding(FMargin(0.0f, 2.0f, 2.0f, 0.0f))
-			[
-				SNew(SBox)
-				[
-					SNew(SVerticalBox)
-					
-					+ SVerticalBox::Slot()
-					[
-						SNew(SButton)
-						.VAlign(VAlign_Center)
-						.HAlign(HAlign_Center)
-						.Text(LOCTEXT("Button_CopyRotation", "Copy Rotation"))
-						.ToolTipText(LOCTEXT("Button_CopyRotation_Tooltip", "Set the virtual texture rotation to match the source actor"))
-						.OnClicked(this, &FRuntimeVirtualTextureComponentDetailsCustomization::SetRotation)
-					]
-
-					+ SVerticalBox::Slot()
-					[
-						SNew(SButton)
-						.VAlign(VAlign_Center)
-						.HAlign(HAlign_Center)
-						.Text(LOCTEXT("Button_CopyBounds", "Copy Bounds"))
-						.ToolTipText(LOCTEXT("Button_CopyBounds_Tooltip", "Set the virtual texture transform so that it includes the full bounds of the source actor"))
-						.OnClicked(this, &FRuntimeVirtualTextureComponentDetailsCustomization::SetTransformToBounds)
-					]
-				]
-			]
-		]
+		SNew(SButton)
+		.VAlign(VAlign_Center)
+		.HAlign(HAlign_Center)
+		.ContentPadding(2)
+		.Text(LOCTEXT("Button_Build", "Build"))
+		.ToolTipText(LOCTEXT("Button_Build_Tooltip", "Build the low mips as streaming virtual texture data"))
+		.OnClicked(this, &FRuntimeVirtualTextureComponentDetailsCustomization::BuildStreamedMips)
 	];
 
-	// Use existing property to add build button
-	TSharedPtr<IPropertyHandle> LowMipsValue = DetailBuilder.GetProperty("StreamLowMips");
-	DetailBuilder.HideProperty(LowMipsValue);
-
-	IDetailCategoryBuilder& VirtualTextureCategory = DetailBuilder.EditCategory("VirtualTextureBuild", FText::GetEmpty());
-	VirtualTextureCategory.AddCustomRow(LowMipsValue->GetPropertyDisplayName())
+	VirtualTextureCategory
+	.AddCustomRow(LOCTEXT("Button_BuildDebugStreamingMips", "Build Debug Streaming Mips"), true)
 	.NameContent()
 	[
-		LowMipsValue->CreatePropertyNameWidget()
+		SNew(STextBlock)
+		.Font(IDetailLayoutBuilder::GetDetailFont())
+		.Text(LOCTEXT("Button_BuildDebugStreamingMips", "Build Debug Streaming Mips"))
 	]
 	.ValueContent()
-	.MaxDesiredWidth(TOptional<float>())
+	.MaxDesiredWidth(125.f)
 	[
-		SNew(SHorizontalBox)
-
-		+ SHorizontalBox::Slot()
-		.Padding(5.f, 0.f)
-		[
-			LowMipsValue->CreatePropertyValueWidget()
-		]
-
-		+ SHorizontalBox::Slot()
-		[
-			SNew(SWrapBox)
-
-			+ SWrapBox::Slot()
-			[
-				SNew(SButton)
-				.Text(LOCTEXT("Button_Build", "Build"))
-				.ToolTipText(LOCTEXT("Button_Build_Tooltip", "Build the low mips as streaming virtual texture data"))
-				.OnClicked(this, &FRuntimeVirtualTextureComponentDetailsCustomization::BuildStreamedMips)
-			]
-		]
-
-		+ SHorizontalBox::Slot()
-		[
-			SNew(SWrapBox)
-
-			+ SWrapBox::Slot()
-			[
-				SNew(SButton)
-				.Text(LOCTEXT("Button_BuildDebig", "Build Debug"))
-				.ToolTipText(LOCTEXT("Button_BuildDebug_Tooltip", "Build the low mips with debug data"))
-				.OnClicked(this, &FRuntimeVirtualTextureComponentDetailsCustomization::BuildLowMipsDebug)
-			]
-		]
+		SNew(SButton)
+		.VAlign(VAlign_Center)
+		.HAlign(HAlign_Center)
+		.ContentPadding(2)
+		.Text(LOCTEXT("Button_Build", "Build"))
+		.ToolTipText(LOCTEXT("Button_BuildDebug_Tooltip", "Build the low mips with debug data"))
+		.OnClicked(this, &FRuntimeVirtualTextureComponentDetailsCustomization::BuildLowMipsDebug)
 	];
 }
 
-FReply FRuntimeVirtualTextureComponentDetailsCustomization::SetRotation()
-{
-	FScopedTransaction BakeTransaction(LOCTEXT("Transaction_CopyRotation", "Copy Rotation"));
-	RuntimeVirtualTextureComponent->Modify();
-	RuntimeVirtualTextureComponent->SetRotation();
-	return FReply::Handled();
-}
-
-FReply FRuntimeVirtualTextureComponentDetailsCustomization::SetTransformToBounds()
-{
-	FScopedTransaction BakeTransaction(LOCTEXT("Transaction_CopyBounds", "Copy Bounds"));
-	RuntimeVirtualTextureComponent->Modify();
-	RuntimeVirtualTextureComponent->SetTransformToBounds();
-	return FReply::Handled();
-}
-
 FReply FRuntimeVirtualTextureComponentDetailsCustomization::BuildStreamedMips()
 {
 	bool bOK = RuntimeVirtualTexture::BuildStreamedMips(RuntimeVirtualTextureComponent, ERuntimeVirtualTextureDebugType::None);
diff --git a/Engine/Source/Runtime/Engine/Classes/Components/BoundsCopyComponent.h b/Engine/Source/Runtime/Engine/Classes/Components/BoundsCopyComponent.h
new file mode 100644
--- /dev/null
+++ b/Engine/Source/Runtime/Engine/Classes/Components/BoundsCopyComponent.h
@@ -0,0 +1,30 @@
+// Copyright 1998-2019 Epic Games, Inc. All Rights Reserved.
+
+#pragma once
+
+#include "CoreMinimal.h"
+#include "ActorComponent.h"
+#include "BoundsCopyComponent.generated.h"
+
+/** Component used to copy the bounds of another Actor. */
+UCLASS(ClassGroup = Rendering, HideCategories = (Activation, AssetUserData, Collision, Cooking, Tags))
+class ENGINE_API UBoundsCopyComponent : public UActorComponent
+{
+	GENERATED_UCLASS_BODY()
+
+private:
+	/** Actor to copy the bounds from to set up the transform. */
+	UPROPERTY(EditAnywhere, Category = TransformFromBounds, meta = (DisplayName = "Source Actor", DisplayPriority = 3))
+	TSoftObjectPtr<AActor> BoundsSourceActor = nullptr;
+
+public:
+#if WITH_EDITOR
+	/** Copy the rotation from BoundsSourceActor to this component */
+	UFUNCTION(CallInEditor, Category = TransformFromBounds, meta = (DisplayName = "Copy Rotation", DisplayPriority = 4))
+	void SetRotation();
+
+	/** Set this component transform to include the BoundsSourceActor bounds */
+	UFUNCTION(CallInEditor, Category = TransformFromBounds, meta = (DisplayName = "Copy Bounds", DisplayPriority = 4))
+	void SetTransformToBounds();
+#endif
+};
diff --git a/Engine/Source/Runtime/Engine/Classes/Components/RuntimeVirtualTextureComponent.h b/Engine/Source/Runtime/Engine/Classes/Components/RuntimeVirtualTextureComponent.h
--- a/Engine/Source/Runtime/Engine/Classes/Components/RuntimeVirtualTextureComponent.h
+++ b/Engine/Source/Runtime/Engine/Classes/Components/RuntimeVirtualTextureComponent.h
@@ -10,7 +10,7 @@ class URuntimeVirtualTexture;
 class UVirtualTextureBuilder;
 
 /** Component used to place a URuntimeVirtualTexture in the world. */
-UCLASS(Blueprintable, ClassGroup = Rendering, collapsecategories, hidecategories = (Activation, Collision, Cooking, Mobility, LOD, Object, Physics, Rendering), editinlinenew)
+UCLASS(Blueprintable, ClassGroup = Rendering, HideCategories = (Activation, Collision, Cooking, Mobility, LOD, Object, Physics, Rendering))
 class ENGINE_API URuntimeVirtualTextureComponent : public USceneComponent
 {
 	GENERATED_UCLASS_BODY()
@@ -21,15 +21,11 @@ protected:
 	URuntimeVirtualTexture* VirtualTexture = nullptr;
 
 	/** Texture object containing streamed low mips. */
-	UPROPERTY(EditAnywhere, BlueprintReadOnly, DuplicateTransient, Category = VirtualTexture)
+	UPROPERTY(EditAnywhere, BlueprintReadOnly, DuplicateTransient, Category = VirtualTextureBuild)
 	UVirtualTextureBuilder* StreamingTexture = nullptr;
 
-	/** Actor to copy the bounds from to set up the transform. */
-	UPROPERTY(EditAnywhere, Category = TransformFromBounds, meta = (DisplayName = "Source Actor"))
-	TSoftObjectPtr<AActor> BoundsSourceActor;
-
 	/** Number of low mips to serialize and stream for the virtual texture. This can reduce rendering update cost. */
-	UPROPERTY(EditAnywhere, Category = VirtualTextureBuild, meta = (UIMin = "0", UIMax = "6", DisplayName = "Streaming Mips"))
+	UPROPERTY(EditAnywhere, Category = VirtualTextureBuild, meta = (UIMin = "0", UIMax = "6", DisplayName = "Num Streaming Mips"))
 	int32 StreamLowMips = 0;
 
 	/** Enable Crunch compression. ZLib compression is used when Crunch is disabled. */
@@ -63,12 +59,6 @@ public:
 #if WITH_EDITOR
 	/** Initialize the low mip streaming texture with the passed in size and data. */
 	void InitializeStreamingTexture(uint32 InSizeX, uint32 InSizeY, uint8* InData);
-
-	/** Copy the rotation from BoundsSourceActor to this component. Called by our UI details customization. */
-	void SetRotation();
-
-	/** Set this component transform to include the BoundsSourceActor bounds. Called by our UI details customization. */
-	void SetTransformToBounds();
 #endif
 
 protected:
diff --git a/Engine/Source/Runtime/Engine/Classes/VT/RuntimeVirtualTextureVolume.h b/Engine/Source/Runtime/Engine/Classes/VT/RuntimeVirtualTextureVolume.h
--- a/Engine/Source/Runtime/Engine/Classes/VT/RuntimeVirtualTextureVolume.h
+++ b/Engine/Source/Runtime/Engine/Classes/VT/RuntimeVirtualTextureVolume.h
@@ -7,7 +7,7 @@
 #include "RuntimeVirtualTextureVolume.generated.h"
 
 /** Actor used to place a URuntimeVirtualTexture in the world. */
-UCLASS(hidecategories=(Actor, Collision, Cooking, Input, LOD, Replication), MinimalAPI)
+UCLASS(HideCategories = (Actor, Collision, Cooking, Input, LOD, Physics, Replication, Rendering), MinimalAPI)
 class ARuntimeVirtualTextureVolume : public AActor
 {
 	GENERATED_UCLASS_BODY()
@@ -18,6 +18,10 @@ private:
 	class URuntimeVirtualTextureComponent* VirtualTextureComponent;
 
 #if WITH_EDITORONLY_DATA
+	/** Component that enables Editor bounds copy functionality. */
+	UPROPERTY(VisibleAnywhere, BlueprintReadOnly, Category = Transform, meta = (AllowPrivateAccess = "true"))
+	class UBoundsCopyComponent* BoundsCopyComponent;
+
 	/** Box for visualizing virtual texture extents. */
 	UPROPERTY(Transient)
 	class UBoxComponent* Box = nullptr;
diff --git a/Engine/Source/Runtime/Engine/Private/Components/BoundsCopyComponent.cpp b/Engine/Source/Runtime/Engine/Private/Components/BoundsCopyComponent.cpp
new file mode 100644
--- /dev/null
+++ b/Engine/Source/Runtime/Engine/Private/Components/BoundsCopyComponent.cpp
@@ -0,0 +1,74 @@
+// Copyright 1998-2019 Epic Games, Inc. All Rights Reserved.
+
+#include "Components/BoundsCopyComponent.h"
+
+#include "Components/PrimitiveComponent.h"
+#include "GameFramework/Actor.h"
+#include "SceneInterface.h"
+
+UBoundsCopyComponent::UBoundsCopyComponent(const FObjectInitializer& ObjectInitializer)
+	: Super(ObjectInitializer)
+{
+	bIsEditorOnly = true;
+}
+
+#if WITH_EDITOR
+
+void UBoundsCopyComponent::SetRotation()
+{
+	if (BoundsSourceActor.IsValid())
+	{
+		// Copy the source actor rotation and notify the parent actor
+		GetOwner()->Modify();
+		GetOwner()->SetActorRotation(BoundsSourceActor->GetTransform().GetRotation());
+		GetOwner()->PostEditMove(true);
+	}
+}
+
+void UBoundsCopyComponent::SetTransformToBounds()
+{
+	if (BoundsSourceActor.IsValid())
+	{
+		// Calculate the bounds in our local rotation space translated to the BoundsSourceActor center
+		const FQuat TargetRotation = GetOwner()->ActorToWorld().GetRotation();
+		const FVector InitialPosition = BoundsSourceActor->GetComponentsBoundingBox().GetCenter();
+
+		FTransform LocalTransform;
+		LocalTransform.SetComponents(TargetRotation, InitialPosition, FVector::OneVector);
+		FTransform WorldToLocal = LocalTransform.Inverse();
+
+		FBox BoundBox(ForceInit);
+		for (const UActorComponent* Component : BoundsSourceActor->GetComponents())
+		{
+			// Only gather visual components in the bounds calculation
+			const UPrimitiveComponent* PrimitiveComponent = Cast<const UPrimitiveComponent>(Component);
+			if (PrimitiveComponent != nullptr && PrimitiveComponent->IsRegistered())
+			{
+				const FTransform ComponentToActor = PrimitiveComponent->GetComponentTransform() * WorldToLocal;
+				FBoxSphereBounds LocalSpaceComponentBounds = PrimitiveComponent->CalcBounds(ComponentToActor);
+				if (LocalSpaceComponentBounds.GetBox().GetVolume() > 0.f)
+				{
+					BoundBox += LocalSpaceComponentBounds.GetBox();
+				}
+			}
+		}
+
+		// Create transform from bounds
+		FVector Origin;
+		FVector Extent;
+		BoundBox.GetCenterAndExtents(Origin, Extent);
+
+		Origin = LocalTransform.TransformPosition(Origin);
+		Extent *= FVector(2.f, 2.f, 1.f); // Account for ARuntimeVirtualTextureVolume:Box offset which centers it on origin
+
+		FTransform Transform;
+		Transform.SetComponents(TargetRotation, Origin, Extent);
+
+		// Apply final result and notify the parent actor
+		GetOwner()->Modify();
+		GetOwner()->SetActorTransform(Transform);
+		GetOwner()->PostEditMove(true);
+	}
+}
+
+#endif
diff --git a/Engine/Source/Runtime/Engine/Private/Components/RuntimeVirtualTextureComponent.cpp b/Engine/Source/Runtime/Engine/Private/Components/RuntimeVirtualTextureComponent.cpp
--- a/Engine/Source/Runtime/Engine/Private/Components/RuntimeVirtualTextureComponent.cpp
+++ b/Engine/Source/Runtime/Engine/Private/Components/RuntimeVirtualTextureComponent.cpp
@@ -12,8 +12,6 @@ URuntimeVirtualTextureComponent::URuntimeVirtualTextureComponent(const FObjectIn
 	: Super(ObjectInitializer)
 	, SceneProxy(nullptr)
 {
-	PrimaryComponentTick.bCanEverTick = true;
-	bTickInEditor = true;
 }
 
 bool URuntimeVirtualTextureComponent::IsVisible() const
@@ -161,59 +159,4 @@ void URuntimeVirtualTextureComponent::InitializeStreamingTexture(uint32 InSizeX,
 	}
 }
 
-void URuntimeVirtualTextureComponent::SetRotation()
-{
-	if (BoundsSourceActor.IsValid())
-	{
-		// Copy the source actor rotation and notify the parent actor
-		SetWorldRotation(BoundsSourceActor->GetTransform().GetRotation());
-		GetOwner()->PostEditMove(true);
-	}
-}
-
-void URuntimeVirtualTextureComponent::SetTransformToBounds()
-{
-	if (BoundsSourceActor.IsValid())
-	{
-		// Calculate the bounds in our local rotation space translated to the BoundsSourceActor center
-		const FQuat TargetRotation = GetComponentToWorld().GetRotation();
-		const FVector InitialPosition = BoundsSourceActor->GetComponentsBoundingBox().GetCenter();
-
-		FTransform LocalTransform;
-		LocalTransform.SetComponents(TargetRotation, InitialPosition, FVector::OneVector);
-		FTransform WorldToLocal = LocalTransform.Inverse();
-
-		FBox BoundBox(ForceInit);
-		for (const UActorComponent* Component : BoundsSourceActor->GetComponents())
-		{
-			// Only gather visual components in the bounds calculation
-			const UPrimitiveComponent* PrimitiveComponent = Cast<const UPrimitiveComponent>(Component);
-			if (PrimitiveComponent != nullptr && PrimitiveComponent->IsRegistered())
-			{
-				const FTransform ComponentToActor = PrimitiveComponent->GetComponentTransform() * WorldToLocal;
-				FBoxSphereBounds LocalSpaceComponentBounds = PrimitiveComponent->CalcBounds(ComponentToActor);
-				if (LocalSpaceComponentBounds.GetBox().GetVolume() > 0.f)
-				{
-					BoundBox += LocalSpaceComponentBounds.GetBox();
-				}
-			}
-		}
-
-		// Create transform from bounds
-		FVector Origin;
-		FVector Extent;
-		BoundBox.GetCenterAndExtents(Origin, Extent);
-
-		Origin = LocalTransform.TransformPosition(Origin);
-		Extent *= FVector(2.f, 2.f, 1.f); // Account for ARuntimeVirtualTextureVolume:Box offset which centers it on origin
-
-		FTransform Transform;
-		Transform.SetComponents(TargetRotation, Origin, Extent);
-
-		// Apply final result and notify the parent actor
-		SetWorldTransform(Transform);
-		GetOwner()->PostEditMove(true);
-	}
-}
-
 #endif
diff --git a/Engine/Source/Runtime/Engine/Private/VT/RuntimeVirtualTextureVolume.cpp b/Engine/Source/Runtime/Engine/Private/VT/RuntimeVirtualTextureVolume.cpp
--- a/Engine/Source/Runtime/Engine/Private/VT/RuntimeVirtualTextureVolume.cpp
+++ b/Engine/Source/Runtime/Engine/Private/VT/RuntimeVirtualTextureVolume.cpp
@@ -3,6 +3,7 @@
 #include "VT/RuntimeVirtualTextureVolume.h"
 
 #include "Components/BoxComponent.h"
+#include "Components/BoundsCopyComponent.h"
 #include "Components/RuntimeVirtualTextureComponent.h"
 
 ARuntimeVirtualTextureVolume::ARuntimeVirtualTextureVolume(const FObjectInitializer& ObjectInitializer)
@@ -11,6 +12,9 @@ ARuntimeVirtualTextureVolume::ARuntimeVirtualTextureVolume(const FObjectInitiali
 	RootComponent = VirtualTextureComponent = CreateDefaultSubobject<URuntimeVirtualTextureComponent>(TEXT("VirtualTextureComponent"));
 
 #if WITH_EDITORONLY_DATA
+	// Add bounds copier
+	BoundsCopyComponent = CreateDefaultSubobject<UBoundsCopyComponent>(TEXT("BoundsCopyComponent"));
+
 	// Add box for visualization of bounds
 	Box = CreateDefaultSubobject<UBoxComponent>(TEXT("Bounds"));
 	Box->SetBoxExtent(FVector(0.5f, 0.5f, 1.f), false);
