Fix regression after PVS fix.
Fixes bad data produced by building streaming virtual texture.

diff --git a/Engine/Source/Editor/VirtualTexturingEditor/Private/RuntimeVirtualTextureBuildStreamingMips.cpp b/Engine/Source/Editor/VirtualTexturingEditor/Private/RuntimeVirtualTextureBuildStreamingMips.cpp
--- a/Engine/Source/Editor/VirtualTexturingEditor/Private/RuntimeVirtualTextureBuildStreamingMips.cpp
+++ b/Engine/Source/Editor/VirtualTexturingEditor/Private/RuntimeVirtualTextureBuildStreamingMips.cpp
@@ -92,7 +92,7 @@ namespace
 		for (int32 y = 0; y < TileSize; y++)
 		{
 			memcpy(
-				DestPixels + (SIZE_T)DestStride * (SIZE_T)DestPos[1] + y + DestPos[0],
+				DestPixels + (SIZE_T)DestStride * ((SIZE_T)DestPos[1] + (SIZE_T)y) + DestPos[0],
 				SrcPixels + SrcStride * y,
 				TileSize * sizeof(T));
 		}
