Use RenderPagesStandAlone() when building RVT and VHM MinMax textures.
Plays correctly with GPUScene and stops ensure() firing.

diff --git a/Engine/Plugins/Experimental/VirtualHeightfieldMesh/Source/VirtualHeightfieldMeshEditor/Private/HeightfieldMinMaxTextureBuild.cpp b/Engine/Plugins/Experimental/VirtualHeightfieldMesh/Source/VirtualHeightfieldMeshEditor/Private/HeightfieldMinMaxTextureBuild.cpp
--- a/Engine/Plugins/Experimental/VirtualHeightfieldMesh/Source/VirtualHeightfieldMeshEditor/Private/HeightfieldMinMaxTextureBuild.cpp
+++ b/Engine/Plugins/Experimental/VirtualHeightfieldMesh/Source/VirtualHeightfieldMeshEditor/Private/HeightfieldMinMaxTextureBuild.cpp
@@ -190,6 +190,9 @@ namespace VirtualHeightfieldMesh
 					MaxLevel, MipLevel = 0](FRHICommandListImmediate& RHICmdList)
 				{
 					// Rendering one page at a time, but could batch here?
+					FMemMark Mark(FMemStack::Get());
+					FRDGBuilder GraphBuilder(RHICmdList);
+
 					const FBox2D TileBox(FVector2D(0, 0), FVector2D(TileSize, TileSize));
 					const FIntRect TileRect(0, 0, TileSize, TileSize);
 
@@ -213,18 +216,15 @@ namespace VirtualHeightfieldMesh
 					Desc.PageDescs[0].UVRange = UVRange;
 					Desc.PageDescs[0].vLevel = MipLevel;
 
-					FMemMark Mark(FMemStack::Get());
-					FRDGBuilder GraphBuilder(RHICmdList);
-
-					RenderPages(GraphBuilder, Desc);
+					RenderPagesStandAlone(GraphBuilder, Desc);
 
 					// Downsample page to texel in output
-
 					FRDGTextureRef SrcTexture = GraphBuilder.RegisterExternalTexture(RenderTileResources.GetTileRenderTarget());
 					FRDGTextureRef DstTexture = GraphBuilder.RegisterExternalTexture(RenderTileResources.GetFinalRenderTarget());
 					FRDGTextureUAVRef DstTextureUAV = GraphBuilder.CreateUAV(DstTexture);
-
+					
 					DownsampleMinMaxAndCopy(GraphBuilder, SrcTexture, FIntPoint(TileSize, TileSize), DstTextureUAV, FIntPoint(TileX, TileY));
+					
 					GraphBuilder.Execute();
 				});
 			}
diff --git a/Engine/Source/Editor/VirtualTexturingEditor/Private/RuntimeVirtualTextureBuildStreamingMips.cpp b/Engine/Source/Editor/VirtualTexturingEditor/Private/RuntimeVirtualTextureBuildStreamingMips.cpp
--- a/Engine/Source/Editor/VirtualTexturingEditor/Private/RuntimeVirtualTextureBuildStreamingMips.cpp
+++ b/Engine/Source/Editor/VirtualTexturingEditor/Private/RuntimeVirtualTextureBuildStreamingMips.cpp
@@ -225,55 +225,50 @@ namespace RuntimeVirtualTexture
 						RHICmdList.Transition(FRHITransitionInfo(RenderTileResources.GetRenderTarget(Layer), ERHIAccess::Unknown, ERHIAccess::RTV));
 					}
 
-					RuntimeVirtualTexture::FRenderPageBatchDesc Desc;
-					Desc.Scene = Scene->GetRenderScene();
-					Desc.RuntimeVirtualTextureMask = 1 << VirtualTextureSceneIndex;
-					Desc.UVToWorld = Transform;
-					Desc.WorldBounds = Bounds;
-					Desc.MaterialType = MaterialType;
-					Desc.MaxLevel = MaxLevel;
-					Desc.bClearTextures = true;
-					Desc.bIsThumbnails = false;
-					Desc.DebugType = DebugType;
-					Desc.NumPageDescs = 1;
-					Desc.Targets[0].Texture = RenderTileResources.GetRenderTarget(0);
-					Desc.Targets[1].Texture = RenderTileResources.GetRenderTarget(1);
-					Desc.Targets[2].Texture = RenderTileResources.GetRenderTarget(2);
-					Desc.PageDescs[0].DestBox[0] = TileBox;
-					Desc.PageDescs[0].DestBox[1] = TileBox;
-					Desc.PageDescs[0].DestBox[2] = TileBox;
-					Desc.PageDescs[0].UVRange = UVRange;
-					Desc.PageDescs[0].vLevel = RenderLevel;
-
 					{
 						FMemMark MemMark(FMemStack::Get());
 						FRDGBuilder GraphBuilder(RHICmdList);
-						RuntimeVirtualTexture::RenderPages(GraphBuilder, Desc);
-						GraphBuilder.Execute();
-					}
 
-					// Transition render targets for copying
-					for (int32 Layer = 0; Layer < NumLayers; Layer++)
-					{
-						RHICmdList.Transition(FRHITransitionInfo(RenderTileResources.GetRenderTarget(Layer), ERHIAccess::RTV, ERHIAccess::SRVGraphics));
+						RuntimeVirtualTexture::FRenderPageBatchDesc Desc;
+						Desc.Scene = Scene->GetRenderScene();
+						Desc.RuntimeVirtualTextureMask = 1 << VirtualTextureSceneIndex;
+						Desc.UVToWorld = Transform;
+						Desc.WorldBounds = Bounds;
+						Desc.MaterialType = MaterialType;
+						Desc.MaxLevel = MaxLevel;
+						Desc.bClearTextures = true;
+						Desc.bIsThumbnails = false;
+						Desc.DebugType = DebugType;
+						Desc.NumPageDescs = 1;
+						Desc.Targets[0].Texture = RenderTileResources.GetRenderTarget(0);
+						Desc.Targets[1].Texture = RenderTileResources.GetRenderTarget(1);
+						Desc.Targets[2].Texture = RenderTileResources.GetRenderTarget(2);
+						Desc.PageDescs[0].DestBox[0] = TileBox;
+						Desc.PageDescs[0].DestBox[1] = TileBox;
+						Desc.PageDescs[0].DestBox[2] = TileBox;
+						Desc.PageDescs[0].UVRange = UVRange;
+						Desc.PageDescs[0].vLevel = RenderLevel;
+
+						RuntimeVirtualTexture::RenderPagesStandAlone(GraphBuilder, Desc);
+
+						GraphBuilder.Execute();
 					}
 
 					// Copy to staging
 					for (int32 Layer = 0; Layer < NumLayers; Layer++)
 					{
+						RHICmdList.Transition(FRHITransitionInfo(RenderTileResources.GetRenderTarget(Layer), ERHIAccess::RTV, ERHIAccess::CopySrc));
 						RHICmdList.CopyTexture(RenderTileResources.GetRenderTarget(Layer), RenderTileResources.GetStagingTexture(Layer), FRHICopyTextureInfo());
 					}
 
-					//todo[vt]: Insert fence for immediate read back. But is there no API to wait on it?
 					RHICmdList.WriteGPUFence(RenderTileResources.GetFence());
-					RHICmdList.ImmediateFlush(EImmediateFlushType::FlushRHIThread);
 
 					// Read back tile data and copy into final destination
 					for (int32 Layer = 0; Layer < NumLayers; Layer++)
 					{
 						void* TilePixels = nullptr;
 						int32 OutWidth, OutHeight;
-						RHICmdList.MapStagingSurface(RenderTileResources.GetStagingTexture(Layer), TilePixels, OutWidth, OutHeight);
+						RHICmdList.MapStagingSurface(RenderTileResources.GetStagingTexture(Layer), RenderTileResources.GetFence(), TilePixels, OutWidth, OutHeight);
 						check(TilePixels != nullptr);
 						check(OutHeight == TileSize);
 
diff --git a/Engine/Source/Editor/VirtualTexturingEditor/Private/RuntimeVirtualTextureThumbnailRenderer.cpp b/Engine/Source/Editor/VirtualTexturingEditor/Private/RuntimeVirtualTextureThumbnailRenderer.cpp
--- a/Engine/Source/Editor/VirtualTexturingEditor/Private/RuntimeVirtualTextureThumbnailRenderer.cpp
+++ b/Engine/Source/Editor/VirtualTexturingEditor/Private/RuntimeVirtualTextureThumbnailRenderer.cpp
@@ -81,6 +81,9 @@ void URuntimeVirtualTextureThumbnailRenderer::Draw(UObject* Object, int32 X, int
 	{
 		FMaterialRenderProxy::UpdateDeferredCachedUniformExpressions();
 
+		FMemMark MemMark(FMemStack::Get());
+		FRDGBuilder GraphBuilder(RHICmdList);
+
 		RuntimeVirtualTexture::FRenderPageBatchDesc Desc;
 		Desc.Scene = Scene->GetRenderScene();
 		Desc.RuntimeVirtualTextureMask = 1 << VirtualTextureSceneIndex;
@@ -97,6 +100,8 @@ void URuntimeVirtualTextureThumbnailRenderer::Draw(UObject* Object, int32 X, int
 		Desc.PageDescs[0].UVRange = FBox2D(FVector2D(0, 0), FVector2D(1, 1));
 		Desc.PageDescs[0].vLevel = MaxLevel;
 
-		RuntimeVirtualTexture::RenderPagesStandAlone(RHICmdList, Desc);
+		RuntimeVirtualTexture::RenderPagesStandAlone(GraphBuilder, Desc);
+
+		GraphBuilder.Execute();
 	});
 }
diff --git a/Engine/Source/Runtime/Renderer/Private/VT/RuntimeVirtualTextureRender.cpp b/Engine/Source/Runtime/Renderer/Private/VT/RuntimeVirtualTextureRender.cpp
--- a/Engine/Source/Runtime/Renderer/Private/VT/RuntimeVirtualTextureRender.cpp
+++ b/Engine/Source/Runtime/Renderer/Private/VT/RuntimeVirtualTextureRender.cpp
@@ -1156,11 +1156,8 @@ namespace RuntimeVirtualTexture
 		}
 	}
 
-	void RenderPagesStandAlone(FRHICommandListImmediate& RHICmdList, FRenderPageBatchDesc const& InDesc)
+	void RenderPagesStandAlone(FRDGBuilder& GraphBuilder, FRenderPageBatchDesc const& InDesc)
 	{
-		FMemMark MemMark(FMemStack::Get());
-		FRDGBuilder GraphBuilder(RHICmdList);
-
 		// This is required to collect dynamic primitives from the views (not used here, but we must provide one).
 		FGPUSceneDynamicContext GPUSceneDynamicContext(InDesc.Scene->GPUScene);
 
@@ -1170,7 +1167,6 @@ namespace RuntimeVirtualTexture
 		FGPUSceneScopeBeginEndHelper GPUSceneScopeBeginEndHelper(InDesc.Scene->GPUScene, GPUSceneDynamicContext, *InDesc.Scene);
 		InDesc.Scene->GPUScene.Update(GraphBuilder, *InDesc.Scene);
 		RenderPages(GraphBuilder, InDesc);
-		GraphBuilder.Execute();
 	}
 
 	void RenderPages(FRHICommandListImmediate& RHICmdList, FRenderPageBatchDesc const& InDesc)
diff --git a/Engine/Source/Runtime/Renderer/Public/VT/RuntimeVirtualTextureRender.h b/Engine/Source/Runtime/Renderer/Public/VT/RuntimeVirtualTextureRender.h
--- a/Engine/Source/Runtime/Renderer/Public/VT/RuntimeVirtualTextureRender.h
+++ b/Engine/Source/Runtime/Renderer/Public/VT/RuntimeVirtualTextureRender.h
@@ -86,10 +86,10 @@ namespace RuntimeVirtualTexture
 	RENDERER_API void RenderPages(FRDGBuilder& GraphBuilder, FRenderPageBatchDesc const& InDesc);
 
 	/** 
-	 * Render a batch of pages for a runtime virtual texture. Performs the required scene rendering setup and 
-	 * creates a graph builder, such that it can be called from a render thread task.
+	 * Render a batch of pages for a runtime virtual texture. 
+	 * Performs the required scene rendering setup such that it can be called from a render thread task.
 	 */
-	RENDERER_API void RenderPagesStandAlone(FRHICommandListImmediate& RHICmdList, FRenderPageBatchDesc const& InDesc);
+	RENDERER_API void RenderPagesStandAlone(FRDGBuilder& GraphBuilder, FRenderPageBatchDesc const& InDesc);
 
 	UE_DEPRECATED(5.0, "This method has been refactored to use an FRDGBuilder instead.")
 	RENDERER_API void RenderPages(FRHICommandListImmediate& RHICmdList, FRenderPageBatchDesc const& InDesc);
