Fix RVT set to bounds bug where bounds could end up including unparented debug components.

diff --git a/Engine/Source/Editor/VirtualTexturingEditor/Private/RuntimeVirtualTextureSetBounds.cpp b/Engine/Source/Editor/VirtualTexturingEditor/Private/RuntimeVirtualTextureSetBounds.cpp
--- a/Engine/Source/Editor/VirtualTexturingEditor/Private/RuntimeVirtualTextureSetBounds.cpp
+++ b/Engine/Source/Editor/VirtualTexturingEditor/Private/RuntimeVirtualTextureSetBounds.cpp
@@ -33,7 +33,7 @@ namespace RuntimeVirtualTexture
 		FBox Bounds(ForceInit);
 		for (TObjectIterator<UPrimitiveComponent> It(RF_ClassDefaultObject, true, EInternalObjectFlags::PendingKill); It; ++It)
 		{
-			bool bUseBounds = It->GetOwner() == BoundsAlignActor.Get();
+			bool bUseBounds = BoundsAlignActor.IsValid() && It->GetOwner() == BoundsAlignActor.Get();
 
 			TArray<URuntimeVirtualTexture*> const& VirtualTextures = It->GetRuntimeVirtualTextures();
 			for (int32 Index = 0; !bUseBounds && Index < VirtualTextures.Num(); ++Index) 
@@ -54,7 +54,7 @@ namespace RuntimeVirtualTexture
 			}
 		}
 
-		// Calulcate the transform to fit the bounds.
+		// Calculate the transform to fit the bounds.
 		FTransform Transform;
 		const FVector LocalPosition = Bounds.Min;
 		const FVector WorldPosition = LocalTransform.TransformPosition(LocalPosition);
