Fix crash when building RVT streaming mips.
Fence needs to be explicitly cleared after recent RHI parallel rendering updates.

diff --git a/Engine/Source/Editor/VirtualTexturingEditor/Private/RuntimeVirtualTextureBuildStreamingMips.cpp b/Engine/Source/Editor/VirtualTexturingEditor/Private/RuntimeVirtualTextureBuildStreamingMips.cpp
--- a/Engine/Source/Editor/VirtualTexturingEditor/Private/RuntimeVirtualTextureBuildStreamingMips.cpp
+++ b/Engine/Source/Editor/VirtualTexturingEditor/Private/RuntimeVirtualTextureBuildStreamingMips.cpp
@@ -270,6 +270,7 @@ namespace RuntimeVirtualTexture
 						RHICmdList.CopyTexture(RenderTileResources.GetRenderTarget(Layer), RenderTileResources.GetStagingTexture(Layer), FRHICopyTextureInfo());
 					}
 
+					RenderTileResources.GetFence()->Clear();
 					RHICmdList.WriteGPUFence(RenderTileResources.GetFence());
 
 					// Read back tile data and copy into final destination
