Editor - Fixing an ASan bug in the category paths, need to declare them statically so that the compiler doesn't free the list too early.
Editor - Upgrading the Texture AssetTypeActions to be AssetDefinitions.

diff --git a/Engine/Plugins/Editor/EngineAssetDefinitions/Source/EngineAssetDefinitions.Build.cs b/Engine/Plugins/Editor/EngineAssetDefinitions/Source/EngineAssetDefinitions.Build.cs
--- a/Engine/Plugins/Editor/EngineAssetDefinitions/Source/EngineAssetDefinitions.Build.cs
+++ b/Engine/Plugins/Editor/EngineAssetDefinitions/Source/EngineAssetDefinitions.Build.cs
@@ -21,11 +21,28 @@ public class EngineAssetDefinitions : ModuleRules
 				"UnrealEd",
 				"ToolMenus",
 				"CoreUObject",
+				"RHI",
 				"Engine",
+				"AssetRegistry",
 				"AssetTools",
 				"SlateCore",
+				"InputCore",
 				"Slate",
-				"Kismet"
+				"Kismet",
+				"TextureEditor",
+				"MaterialEditor",
+			}
+		);
+		
+		PrivateIncludePathModuleNames.AddRange(
+			new string[] {
+				"VirtualTexturingEditor",
+			}
+		);
+
+		DynamicallyLoadedModuleNames.AddRange(
+			new string[] {
+				"VirtualTexturingEditor",
 			}
 		);
 	}
diff --git a/Engine/Plugins/Editor/EngineAssetDefinitions/Source/Private/AssetDefinition_CanvasRenderTarget2D.h b/Engine/Plugins/Editor/EngineAssetDefinitions/Source/Private/AssetDefinition_CanvasRenderTarget2D.h
new file mode 100644
--- /dev/null
+++ b/Engine/Plugins/Editor/EngineAssetDefinitions/Source/Private/AssetDefinition_CanvasRenderTarget2D.h
@@ -0,0 +1,21 @@
+// Copyright Epic Games, Inc. All Rights Reserved.
+
+#pragma once
+
+#include "CoreMinimal.h"
+#include "Engine/CanvasRenderTarget2D.h"
+#include "AssetDefinition_TextureRenderTarget2D.h"
+
+#include "AssetDefinition_CanvasRenderTarget2D.generated.h"
+
+UCLASS()
+class ENGINEASSETDEFINITIONS_API UAssetDefinition_CanvasRenderTarget2D : public UAssetDefinition_TextureRenderTarget2D
+{
+	GENERATED_BODY()
+
+public:
+	// UAssetDefinition Begin
+	virtual FText GetAssetDisplayName() const override { return NSLOCTEXT("AssetTypeActions", "AssetTypeActions_CanvasRenderTarget2D", "Canvas Render Target"); }
+	virtual TSoftClassPtr<UObject> GetAssetClass() const override { return UCanvasRenderTarget2D::StaticClass(); }
+	// UAssetDefinition End
+};
diff --git a/Engine/Plugins/Editor/EngineAssetDefinitions/Source/Private/AssetDefinition_CurveLinearColorAtlas.h b/Engine/Plugins/Editor/EngineAssetDefinitions/Source/Private/AssetDefinition_CurveLinearColorAtlas.h
new file mode 100644
--- /dev/null
+++ b/Engine/Plugins/Editor/EngineAssetDefinitions/Source/Private/AssetDefinition_CurveLinearColorAtlas.h
@@ -0,0 +1,28 @@
+// Copyright Epic Games, Inc. All Rights Reserved.
+
+#pragma once
+
+#include "CoreMinimal.h"
+#include "Curves/CurveLinearColorAtlas.h"
+#include "AssetDefinition_Texture.h"
+
+#include "AssetDefinition_CurveLinearColorAtlas.generated.h"
+
+UCLASS()
+class ENGINEASSETDEFINITIONS_API UAssetDefinition_CurveLinearColorAtlas : public UAssetDefinition_Texture
+{
+	GENERATED_BODY()
+
+public:
+	// UAssetDefinition Begin
+	virtual FText GetAssetDisplayName() const override { return NSLOCTEXT("AssetTypeActions", "AssetTypeActions_CurveLinearColorAtlas", "Curve Atlas"); }
+	virtual FLinearColor GetAssetColor() const override { return FLinearColor(FColor(128,64,64)); }
+	virtual TSoftClassPtr<UObject> GetAssetClass() const override { return UCurveLinearColorAtlas::StaticClass(); }
+	virtual bool CanImport() const override { return false; }
+	virtual TConstArrayView<FAssetCategoryPath> GetAssetCategories() const override
+	{
+		static const auto Categories ={ EAssetCategoryPaths::Misc };
+		return Categories;
+	}
+	// UAssetDefinition End
+};
diff --git a/Engine/Plugins/Editor/EngineAssetDefinitions/Source/Private/AssetDefinition_Texture.cpp b/Engine/Plugins/Editor/EngineAssetDefinitions/Source/Private/AssetDefinition_Texture.cpp
new file mode 100644
--- /dev/null
+++ b/Engine/Plugins/Editor/EngineAssetDefinitions/Source/Private/AssetDefinition_Texture.cpp
@@ -0,0 +1,242 @@
+// Copyright Epic Games, Inc. All Rights Reserved.
+
+#include "AssetDefinition_Texture.h"
+
+#include "ContentBrowserMenuContexts.h"
+#include "Engine/Texture2D.h"
+#include "Misc/PackageName.h"
+#include "Styling/AppStyle.h"
+#include "Materials/Material.h"
+#include "Materials/MaterialInstance.h"
+#include "Factories/MaterialFactoryNew.h"
+#include "Interfaces/ITextureEditorModule.h"
+#include "IContentBrowserSingleton.h"
+#include "ContentBrowserModule.h"
+#include "Factories/SubUVAnimationFactory.h"
+#include "Particles/SubUVAnimation.h"
+#include "AssetRegistry/AssetRegistryModule.h"
+
+#include "ToolMenus.h"
+#include "Misc/FeedbackContext.h"
+#include "IAssetTools.h"
+#include "VirtualTexturingEditorModule.h"
+#include "Algo/AnyOf.h"
+#include "AssetRegistry/AssetRegistryHelpers.h"
+
+#define LOCTEXT_NAMESPACE "UAssetDefinition_Texture"
+
+EAssetCommandResult UAssetDefinition_Texture::OpenAssets(const FAssetOpenArgs& OpenArgs) const
+{
+	const EToolkitMode::Type Mode = OpenArgs.ToolkitHost.IsValid() ? EToolkitMode::WorldCentric : EToolkitMode::Standalone;
+
+	for (UTexture* Texture : OpenArgs.LoadObjects<UTexture>())
+	{
+		ITextureEditorModule* TextureEditorModule = &FModuleManager::LoadModuleChecked<ITextureEditorModule>("TextureEditor");
+		TextureEditorModule->CreateTextureEditor(Mode, OpenArgs.ToolkitHost, Texture);
+	}
+	
+	return EAssetCommandResult::Handled;
+}
+
+// Menu Extensions
+//--------------------------------------------------------------------
+
+namespace MenuExtension_Texture
+{
+	void ExecuteCreateMaterial(const FToolMenuContext& InContext)
+	{
+		const UContentBrowserAssetContextMenuContext* CBContext = UContentBrowserAssetContextMenuContext::FindContextWithAssets(InContext);
+
+		const FString DefaultSuffix = TEXT("_Mat");
+
+		if ( const FAssetData* TextureAsset = CBContext->GetSingleSelectedAssetOfType(UTexture::StaticClass()) )
+        {
+			if (UTexture* Texture = CastChecked<UTexture>(TextureAsset->GetAsset()))
+			{
+				// Determine an appropriate name
+				FString Name;
+				FString PackagePath;
+				IAssetTools::Get().CreateUniqueAssetName(Texture->GetOutermost()->GetName(), DefaultSuffix, PackagePath, Name);
+
+				// Create the factory used to generate the asset
+				UMaterialFactoryNew* Factory = NewObject<UMaterialFactoryNew>();
+				Factory->InitialTexture = Texture;
+
+				FContentBrowserModule& ContentBrowserModule = FModuleManager::LoadModuleChecked<FContentBrowserModule>("ContentBrowser");
+				ContentBrowserModule.Get().CreateNewAsset(Name, FPackageName::GetLongPackagePath(PackagePath), UMaterial::StaticClass(), Factory);
+			}
+		}
+		else
+		{
+			TArray<UObject*> ObjectsToSync;
+
+			for (UTexture* Texture : CBContext->LoadSelectedObjects<UTexture>())
+			{
+				FString Name;
+				FString PackageName;
+				IAssetTools::Get().CreateUniqueAssetName(Texture->GetOutermost()->GetName(), DefaultSuffix, PackageName, Name);
+
+				// Create the factory used to generate the asset
+				UMaterialFactoryNew* Factory = NewObject<UMaterialFactoryNew>();
+				Factory->InitialTexture = Texture;
+
+				UObject* NewAsset = IAssetTools::Get().CreateAsset(Name, FPackageName::GetLongPackagePath(PackageName), UMaterial::StaticClass(), Factory);
+
+				if ( NewAsset )
+				{
+					ObjectsToSync.Add(NewAsset);
+				}
+			}
+
+			if ( ObjectsToSync.Num() > 0 )
+			{
+				IAssetTools::Get().SyncBrowserToAssets(ObjectsToSync);
+			}
+		}
+	}
+
+	static void ExecuteConvertToVirtualTexture(const FToolMenuContext& InContext)
+	{
+		if (const UContentBrowserAssetContextMenuContext* CBContext = UContentBrowserAssetContextMenuContext::FindContextWithAssets(InContext))
+		{
+			IVirtualTexturingEditorModule* Module = FModuleManager::Get().GetModulePtr<IVirtualTexturingEditorModule>("VirtualTexturingEditor");
+			Module->ConvertVirtualTexturesWithDialog(CBContext->LoadSelectedObjects<UTexture2D>(), false);
+		}
+	}
+
+	static void ExecuteConvertToRegularTexture(const FToolMenuContext& InContext)
+	{
+		if (const UContentBrowserAssetContextMenuContext* CBContext = UContentBrowserAssetContextMenuContext::FindContextWithAssets(InContext))
+		{
+			IVirtualTexturingEditorModule* Module = FModuleManager::Get().GetModulePtr<IVirtualTexturingEditorModule>("VirtualTexturingEditor");
+			Module->ConvertVirtualTexturesWithDialog(CBContext->LoadSelectedObjects<UTexture2D>(), true);
+		}
+	}
+
+	static void ExecuteCreateSubUVAnimation(TArray<TWeakObjectPtr<UTexture>> Objects)
+	{
+		const FString DefaultSuffix = TEXT("_SubUV");
+
+		if ( Objects.Num() == 1 )
+		{
+			UTexture2D* Object = Cast<UTexture2D>(Objects[0].Get());
+
+			if ( Object )
+			{
+				// Determine an appropriate name
+				FString Name;
+				FString PackagePath;
+				IAssetTools::Get().CreateUniqueAssetName(Object->GetOutermost()->GetName(), DefaultSuffix, PackagePath, Name);
+
+				// Create the factory used to generate the asset
+				USubUVAnimationFactory* Factory = NewObject<USubUVAnimationFactory>();
+				Factory->InitialTexture = Object;
+
+				FContentBrowserModule& ContentBrowserModule = FModuleManager::LoadModuleChecked<FContentBrowserModule>("ContentBrowser");
+				ContentBrowserModule.Get().CreateNewAsset(Name, FPackageName::GetLongPackagePath(PackagePath), USubUVAnimation::StaticClass(), Factory);
+			}
+		}
+		else
+		{
+			TArray<UObject*> ObjectsToSync;
+
+			for (auto ObjIt = Objects.CreateConstIterator(); ObjIt; ++ObjIt)
+			{
+				UTexture2D* Object = Cast<UTexture2D>((*ObjIt).Get());
+
+				if ( Object )
+				{
+					FString Name;
+					FString PackageName;
+					IAssetTools::Get().CreateUniqueAssetName(Object->GetOutermost()->GetName(), DefaultSuffix, PackageName, Name);
+
+					// Create the factory used to generate the asset
+					USubUVAnimationFactory* Factory = NewObject<USubUVAnimationFactory>();
+					Factory->InitialTexture = Object;
+					
+					UObject* NewAsset = IAssetTools::Get().CreateAsset(Name, FPackageName::GetLongPackagePath(PackageName), USubUVAnimation::StaticClass(), Factory);
+
+					if ( NewAsset )
+					{
+						ObjectsToSync.Add(NewAsset);
+					}
+				}
+			}
+
+			if ( ObjectsToSync.Num() > 0 )
+			{
+				IAssetTools::Get().SyncBrowserToAssets(ObjectsToSync);
+			}
+		}
+	}
+
+	static void ExecuteFindMaterials(const FToolMenuContext& InContext)
+	{
+		if (const UContentBrowserAssetContextMenuContext* CBContext = UContentBrowserAssetContextMenuContext::FindContextWithAssets(InContext))
+		{
+			TArray<FAssetData> MaterialsUsingTexture;
+
+			// This finds "Material like" objects. It's called material in the UI string but this generally
+			// seems more useful
+			if ( const FAssetData* TextureAsset = CBContext->GetSingleSelectedAssetOfType(UTexture::StaticClass()) )
+			{
+				UAssetRegistryHelpers::FindReferencersOfAssetOfClass(TextureAsset->PackageName, { UMaterialInterface::StaticClass(), UMaterialFunction::StaticClass() }, MaterialsUsingTexture);
+			}
+
+			if (MaterialsUsingTexture.Num() > 0)
+			{
+				IAssetTools::Get().SyncBrowserToAssets(MaterialsUsingTexture);
+			}
+		}
+	}
+
+	static FDelayedAutoRegisterHelper DelayedAutoRegister(EDelayedRegisterRunPhase::EndOfEngineInit, []{ 
+		UToolMenus::RegisterStartupCallback(FSimpleMulticastDelegate::FDelegate::CreateLambda([]()
+		{
+			FToolMenuOwnerScoped OwnerScoped(UE_MODULE_NAME);
+			UToolMenu* Menu = UE::ContentBrowser::ExtendToolMenu_AssetContextMenu(UTexture::StaticClass());
+		
+			FToolMenuSection& Section = Menu->FindOrAddSection("GetAssetActions");
+				Section.AddDynamicEntry("GetAssetActions_UTexture", FNewToolMenuSectionDelegate::CreateLambda([](FToolMenuSection& InSection)
+				{
+					if (const UContentBrowserAssetContextMenuContext* Context = UContentBrowserAssetContextMenuContext::FindContextWithAssets(InSection))
+					{
+						{
+							const TAttribute<FText> Label = LOCTEXT("Texture_CreateMaterial", "Create Material");
+							const TAttribute<FText> ToolTip = LOCTEXT("Texture_CreateMaterialTooltip", "Creates a new material using this texture.");
+							const FSlateIcon Icon = FSlateIcon(FAppStyle::GetAppStyleSetName(), "ClassIcon.Material");
+							const FToolMenuExecuteAction UIAction = FToolMenuExecuteAction::CreateStatic(&ExecuteCreateMaterial);
+							InSection.AddMenuEntry("Texture_CreateMaterial", Label, ToolTip, Icon, UIAction);
+						}
+					
+						const bool bHasVirtualTextures =
+							Algo::AnyOf(Context->SelectedAssets, [](const FAssetData& AssetData){ 
+								bool VirtualTextured = false;
+								AssetData.GetTagValue<bool>("VirtualTextureStreaming", VirtualTextured);
+								return VirtualTextured;
+							});
+					
+						if (bHasVirtualTextures)
+						{
+							const TAttribute<FText> Label = LOCTEXT("Texture_ConvertToVT", "Convert to Virtual Texture");
+							const TAttribute<FText> ToolTip = LOCTEXT("Texture_ConvertToVTTooltip", "Converts this texture to a virtual texture if it fits the size limit imposed in the texture importer settings.");
+							const FSlateIcon Icon = FSlateIcon(FAppStyle::GetAppStyleSetName(), "ClassIcon.Texture2D");
+							const FToolMenuExecuteAction UIAction = FToolMenuExecuteAction::CreateStatic(&ExecuteConvertToVirtualTexture);
+							InSection.AddMenuEntry("Texture_ConvertToVT", Label, ToolTip, Icon, UIAction);
+						}
+						
+						if ( Context->SelectedAssets.Num() == 1 )
+						{
+							const TAttribute<FText> Label = LOCTEXT("Texture_FindMaterials", "Find Materials Using This");
+							const TAttribute<FText> ToolTip = LOCTEXT("Texture_FindMaterialsTooltip", "Finds all materials that use this material in the content browser.");
+							const FSlateIcon Icon = FSlateIcon(FAppStyle::GetAppStyleSetName(), "Icons.Find");
+							const FToolMenuExecuteAction UIAction = FToolMenuExecuteAction::CreateStatic(&ExecuteFindMaterials);
+							InSection.AddMenuEntry("Texture_FindMaterials", Label, ToolTip, Icon, UIAction);
+						}
+					}
+				}));
+		}));
+	});
+}
+
+#undef LOCTEXT_NAMESPACE
diff --git a/Engine/Plugins/Editor/EngineAssetDefinitions/Source/Private/AssetDefinition_Texture.h b/Engine/Plugins/Editor/EngineAssetDefinitions/Source/Private/AssetDefinition_Texture.h
new file mode 100644
--- /dev/null
+++ b/Engine/Plugins/Editor/EngineAssetDefinitions/Source/Private/AssetDefinition_Texture.h
@@ -0,0 +1,30 @@
+// Copyright Epic Games, Inc. All Rights Reserved.
+
+#pragma once
+
+#include "CoreMinimal.h"
+#include "Engine/Texture.h"
+#include "AssetDefinitionDefault.h"
+
+#include "AssetDefinition_Texture.generated.h"
+
+UCLASS(Abstract)
+class ENGINEASSETDEFINITIONS_API UAssetDefinition_Texture : public UAssetDefinitionDefault
+{
+	GENERATED_BODY()
+
+public:
+	// UAssetDefinition Begin
+	virtual FText GetAssetDisplayName() const override { return NSLOCTEXT("AssetTypeActions", "AssetTypeActions_Texture", "BaseTexture"); }
+	virtual FLinearColor GetAssetColor() const override { return FLinearColor(FColor(192, 64, 64)); }
+	virtual TSoftClassPtr<UObject> GetAssetClass() const override { return UTexture::StaticClass(); }
+	virtual TConstArrayView<FAssetCategoryPath> GetAssetCategories() const override
+	{
+		static const auto Categories = { EAssetCategoryPaths::Texture };
+		return Categories;
+	}
+	virtual bool CanImport() const override { return true; }
+	
+	virtual EAssetCommandResult OpenAssets(const FAssetOpenArgs& OpenArgs) const override;
+	// UAssetDefinition End
+};
diff --git a/Engine/Plugins/Editor/EngineAssetDefinitions/Source/Private/AssetDefinition_Texture2D.cpp b/Engine/Plugins/Editor/EngineAssetDefinitions/Source/Private/AssetDefinition_Texture2D.cpp
new file mode 100644
--- /dev/null
+++ b/Engine/Plugins/Editor/EngineAssetDefinitions/Source/Private/AssetDefinition_Texture2D.cpp
@@ -0,0 +1,155 @@
+// Copyright Epic Games, Inc. All Rights Reserved.
+
+#include "AssetDefinition_Texture2D.h"
+
+#include "AssetToolsModule.h"
+#include "ContentBrowserMenuContexts.h"
+#include "ToolMenus.h"
+#include "Misc/PackageName.h"
+#include "Styling/AppStyle.h"
+#include "Factories/SlateBrushAssetFactory.h"
+#include "Slate/SlateBrushAsset.h"
+#include "Factories/VolumeTextureFactory.h"
+#include "Factories/Texture2DArrayFactory.h"
+#include "Engine/VolumeTexture.h"
+#include "Engine/Texture2DArray.h"
+#include "IAssetTools.h"
+#include "IContentBrowserSingleton.h"
+#include "ContentBrowserModule.h"
+
+#define LOCTEXT_NAMESPACE "UAssetDefinition_Texture2D"
+
+namespace MenuExtension_Texture2D
+{
+	void ExecuteCreateSlateBrush(const FToolMenuContext& MenuContext)
+	{
+		const UContentBrowserAssetContextMenuContext* Context = UContentBrowserAssetContextMenuContext::FindContextWithAssets(MenuContext);
+		TArray<UTexture2D*> SelectedTextures = Context->LoadSelectedObjects<UTexture2D>();
+		
+		const FString DefaultSuffix = TEXT("_Brush");
+		
+		TArray<UObject*> ObjectsToSync;
+
+		for (UTexture2D* Texture : SelectedTextures)
+		{
+			// Determine the asset name
+			FString Name;
+			FString PackageName;
+			IAssetTools::Get().CreateUniqueAssetName(Texture->GetOutermost()->GetName(), DefaultSuffix, PackageName, Name);
+
+			// Create the factory used to generate the asset
+			USlateBrushAssetFactory* Factory = NewObject<USlateBrushAssetFactory>();
+			Factory->InitialTexture = Texture;
+
+			FAssetToolsModule& AssetToolsModule = FModuleManager::GetModuleChecked<FAssetToolsModule>("AssetTools");
+			UObject* NewAsset = AssetToolsModule.Get().CreateAsset(Name, FPackageName::GetLongPackagePath(PackageName), USlateBrushAsset::StaticClass(), Factory);
+
+			if( NewAsset )
+			{
+				ObjectsToSync.Add(NewAsset);
+			}
+		}
+
+		if( ObjectsToSync.Num() > 0 )
+		{
+			IAssetTools::Get().SyncBrowserToAssets(ObjectsToSync);
+		}
+	}
+
+	void ExecuteCreateTextureArray(const FToolMenuContext& MenuContext)
+	{
+		const UContentBrowserAssetContextMenuContext* Context = UContentBrowserAssetContextMenuContext::FindContextWithAssets(MenuContext);
+		TArray<UTexture2D*> SelectedTextures = Context->LoadSelectedObjects<UTexture2D>();
+		if(SelectedTextures.Num() > 0)
+		{
+			const FString DefaultSuffix = TEXT("_Array");
+			FString Name;
+			FString PackagePath;
+			IAssetTools::Get().CreateUniqueAssetName(SelectedTextures[0]->GetOutermost()->GetName(), DefaultSuffix, PackagePath, Name);
+		
+			// Create the factory used to generate the asset.
+			UTexture2DArrayFactory* Factory = NewObject<UTexture2DArrayFactory>();
+			Factory->InitialTextures.Empty();
+
+			// Give the selected textures to the factory.
+			for (int32 TextureIndex = 0; TextureIndex < SelectedTextures.Num(); ++TextureIndex)
+			{
+				Factory->InitialTextures.Add(SelectedTextures[TextureIndex]);
+			}
+
+			{
+				FContentBrowserModule& ContentBrowserModule = FModuleManager::LoadModuleChecked<FContentBrowserModule>("ContentBrowser");
+				ContentBrowserModule.Get().CreateNewAsset(Name, FPackageName::GetLongPackagePath(PackagePath), UTexture2DArray::StaticClass(), Factory);
+			}
+		}
+	}
+
+	void ExecuteCreateVolumeTexture(const FToolMenuContext& MenuContext)
+	{
+		const UContentBrowserAssetContextMenuContext* Context = UContentBrowserAssetContextMenuContext::FindContextWithAssets(MenuContext);
+		const FString DefaultSuffix = TEXT("_Volume");
+
+		if (Context->SelectedAssets.Num() == 1)
+		{
+			if (UTexture2D* Object = Cast<UTexture2D>(Context->SelectedAssets[0].GetAsset()))
+			{
+				// Determine the asset name
+				FString Name;
+				FString PackagePath;
+				IAssetTools::Get().CreateUniqueAssetName(Object->GetOutermost()->GetName(), DefaultSuffix, PackagePath, Name);
+
+				// Create the factory used to generate the asset
+				UVolumeTextureFactory* Factory = NewObject<UVolumeTextureFactory>();
+				Factory->InitialTexture = Object;
+				FContentBrowserModule& ContentBrowserModule = FModuleManager::LoadModuleChecked<FContentBrowserModule>("ContentBrowser");
+				ContentBrowserModule.Get().CreateNewAsset(Name, FPackageName::GetLongPackagePath(PackagePath), UVolumeTexture::StaticClass(), Factory);
+			}
+		}
+	}
+
+	static FDelayedAutoRegisterHelper DelayedAutoRegister(EDelayedRegisterRunPhase::EndOfEngineInit, []{ 
+		UToolMenus::RegisterStartupCallback(FSimpleMulticastDelegate::FDelegate::CreateLambda([]()
+		{
+			FToolMenuOwnerScoped OwnerScoped(UE_MODULE_NAME);
+			UToolMenu* Menu = UE::ContentBrowser::ExtendToolMenu_AssetContextMenu(UTexture2D::StaticClass());
+	        
+			FToolMenuSection& Section = Menu->FindOrAddSection("GetAssetActions");
+			Section.AddDynamicEntry("GetAssetActions_UTexture2D", FNewToolMenuSectionDelegate::CreateLambda([](FToolMenuSection& InSection)
+			{
+				const UContentBrowserAssetContextMenuContext* Context = UContentBrowserAssetContextMenuContext::FindContextWithAssets(InSection);
+				
+				{
+					const TAttribute<FText> Label = LOCTEXT("Texture2D_CreateSlateBrush", "Create Slate Brush");
+					const TAttribute<FText> ToolTip = LOCTEXT("Texture2D_CreateSlateBrushToolTip", "Creates a new slate brush using this texture.");
+					const FSlateIcon Icon = FSlateIcon(FAppStyle::GetAppStyleSetName(), "ClassIcon.SlateBrushAsset");
+					const FToolMenuExecuteAction UIAction = FToolMenuExecuteAction::CreateStatic(&ExecuteCreateSlateBrush);
+
+					InSection.AddMenuEntry("Texture2D_CreateSlateBrush", Label, ToolTip, Icon, UIAction);
+				}
+
+				static const auto AllowTextureArrayAssetCreationVar = IConsoleManager::Get().FindTConsoleVariableDataInt(TEXT("r.AllowTexture2DArrayCreation"));
+				if (AllowTextureArrayAssetCreationVar->GetValueOnGameThread() != 0)
+				{
+					const TAttribute<FText> Label = LOCTEXT("Texture_Texture2DArray", "Create Texture Array");
+					const TAttribute<FText> ToolTip = LOCTEXT("Texture_CreateTexture2DArrayTooltip", "Creates a new texture array.");
+					const FSlateIcon Icon = FSlateIcon(FAppStyle::GetAppStyleSetName(), "ClassIcon.Texture2D");
+					const FToolMenuExecuteAction UIAction = FToolMenuExecuteAction::CreateStatic(&ExecuteCreateTextureArray);
+
+					InSection.AddMenuEntry("Texture_Texture2DArray", Label, ToolTip, Icon, UIAction);
+				}
+
+				if (Context->SelectedAssets.Num() == 1)
+                {
+					const TAttribute<FText> Label = LOCTEXT("Texture2D_CreateVolumeTexture", "Create Volume Texture");
+					const TAttribute<FText> ToolTip = LOCTEXT("Texture2D_CreateVolumeTextureToolTip", "Creates a new volume texture using this texture.");
+					const FSlateIcon Icon = FSlateIcon(FAppStyle::GetAppStyleSetName(), "ContentBrowser.AssetActions.VolumeTexture");
+					const FToolMenuExecuteAction UIAction = FToolMenuExecuteAction::CreateStatic(&ExecuteCreateVolumeTexture);
+
+					InSection.AddMenuEntry("Texture2D_CreateVolumeTexture", Label, ToolTip, Icon, UIAction);
+				}
+			}));
+		}));
+	});
+}
+
+#undef LOCTEXT_NAMESPACE
diff --git a/Engine/Plugins/Editor/EngineAssetDefinitions/Source/Private/AssetDefinition_Texture2D.h b/Engine/Plugins/Editor/EngineAssetDefinitions/Source/Private/AssetDefinition_Texture2D.h
new file mode 100644
--- /dev/null
+++ b/Engine/Plugins/Editor/EngineAssetDefinitions/Source/Private/AssetDefinition_Texture2D.h
@@ -0,0 +1,26 @@
+// Copyright Epic Games, Inc. All Rights Reserved.
+
+#pragma once
+
+#include "CoreMinimal.h"
+#include "Engine/Texture2D.h"
+#include "AssetDefinition_Texture.h"
+
+#include "AssetDefinition_Texture2D.generated.h"
+
+UCLASS()
+class ENGINEASSETDEFINITIONS_API UAssetDefinition_Texture2D : public UAssetDefinition_Texture
+{
+	GENERATED_BODY()
+
+public:
+	// UAssetDefinition Begin
+	virtual FText GetAssetDisplayName() const override { return NSLOCTEXT("AssetTypeActions", "AssetTypeActions_Texture2D", "Texture"); }
+	virtual TSoftClassPtr<UObject> GetAssetClass() const override { return UTexture2D::StaticClass(); }
+	virtual TConstArrayView<FAssetCategoryPath> GetAssetCategories() const override
+	{
+		static const auto Categories = { EAssetCategoryPaths::Texture, EAssetCategoryPaths::Basic };
+		return Categories;
+	}
+	// UAssetDefinition End
+};
diff --git a/Engine/Plugins/Editor/EngineAssetDefinitions/Source/Private/AssetDefinition_Texture2DArray.h b/Engine/Plugins/Editor/EngineAssetDefinitions/Source/Private/AssetDefinition_Texture2DArray.h
new file mode 100644
--- /dev/null
+++ b/Engine/Plugins/Editor/EngineAssetDefinitions/Source/Private/AssetDefinition_Texture2DArray.h
@@ -0,0 +1,22 @@
+// Copyright Epic Games, Inc. All Rights Reserved.
+
+#pragma once
+
+#include "CoreMinimal.h"
+#include "Engine/Texture2DArray.h"
+#include "AssetDefinitionDefault.h"
+
+#include "AssetDefinition_Texture2DArray.generated.h"
+
+UCLASS()
+class ENGINEASSETDEFINITIONS_API UAssetDefinition_Texture2DArray : public UAssetDefinition_Texture
+{
+	GENERATED_BODY()
+
+public:
+	// UAssetDefinition Begin
+	virtual FText GetAssetDisplayName() const override { return NSLOCTEXT("AssetTypeActions", "AssetTypeActions_Texture2DArray", "Texture 2D Array"); }
+	virtual FLinearColor GetAssetColor() const override { return FLinearColor(FColor(0, 64, 128)); }
+	virtual TSoftClassPtr<UObject> GetAssetClass() const override { return UTexture2DArray::StaticClass(); }
+	// UAssetDefinition End
+};
diff --git a/Engine/Plugins/Editor/EngineAssetDefinitions/Source/Private/AssetDefinition_TextureCube.cpp b/Engine/Plugins/Editor/EngineAssetDefinitions/Source/Private/AssetDefinition_TextureCube.cpp
new file mode 100644
--- /dev/null
+++ b/Engine/Plugins/Editor/EngineAssetDefinitions/Source/Private/AssetDefinition_TextureCube.cpp
@@ -0,0 +1,72 @@
+// Copyright Epic Games, Inc. All Rights Reserved.
+
+#include "AssetDefinition_TextureCube.h"
+
+#include "AssetToolsModule.h"
+#include "ContentBrowserMenuContexts.h"
+#include "ToolMenus.h"
+#include "Misc/PackageName.h"
+#include "Styling/AppStyle.h"
+#include "Factories/TextureCubeArrayFactory.h"
+#include "Engine/TextureCubeArray.h"
+#include "IAssetTools.h"
+#include "IContentBrowserSingleton.h"
+#include "ContentBrowserModule.h"
+
+#define LOCTEXT_NAMESPACE "UAssetDefinition_TextureCubeArray"
+
+namespace MenuExtension_TextureCubeArray
+{
+	static void ExecuteCreateTextureArray(const FToolMenuContext& MenuContext)
+	{
+		if (const UContentBrowserAssetContextMenuContext* Context = UContentBrowserAssetContextMenuContext::FindContextWithAssets(MenuContext))
+		{
+			TArray<UTextureCube*> SelectedTextureCubes = Context->LoadSelectedObjects<UTextureCube>();
+			if (SelectedTextureCubes.Num() > 0)
+			{
+				const FString DefaultSuffix = TEXT("_Array");
+				FString Name;
+				FString PackagePath;
+				IAssetTools::Get().CreateUniqueAssetName(SelectedTextureCubes[0]->GetOutermost()->GetName(), DefaultSuffix, PackagePath, Name);
+
+				// Create the factory used to generate the asset.
+				UTextureCubeArrayFactory* Factory = NewObject<UTextureCubeArrayFactory>();
+				Factory->InitialTextures.Empty();
+
+				// Give the selected textures to the factory.
+				for (int32 TextureIndex = 0; TextureIndex < SelectedTextureCubes.Num(); ++TextureIndex)
+				{
+					Factory->InitialTextures.Add(SelectedTextureCubes[TextureIndex]);
+				}
+
+				{
+					FContentBrowserModule& ContentBrowserModule = FModuleManager::LoadModuleChecked<FContentBrowserModule>("ContentBrowser");
+					ContentBrowserModule.Get().CreateNewAsset(Name, FPackageName::GetLongPackagePath(PackagePath), UTextureCubeArray::StaticClass(), Factory);
+				}
+			}
+		}
+	}
+	
+	static FDelayedAutoRegisterHelper DelayedAutoRegister(EDelayedRegisterRunPhase::EndOfEngineInit, []{ 
+		UToolMenus::RegisterStartupCallback(FSimpleMulticastDelegate::FDelegate::CreateLambda([]()
+		{
+			FToolMenuOwnerScoped OwnerScoped(UE_MODULE_NAME);
+			UToolMenu* Menu = UE::ContentBrowser::ExtendToolMenu_AssetContextMenu(UTextureCube::StaticClass());
+	        
+			FToolMenuSection& Section = Menu->FindOrAddSection("GetAssetActions");
+			Section.AddDynamicEntry("GetAssetActions_UTextureCube", FNewToolMenuSectionDelegate::CreateLambda([](FToolMenuSection& InSection)
+			{
+				{
+					const TAttribute<FText> Label = LOCTEXT("Texture_TextureCubeArray", "Create Texture Cube Array");
+					const TAttribute<FText> ToolTip = LOCTEXT("Texture_CreateTextureCubeArrayTooltip", "Creates a new texture cube array.");
+					const FSlateIcon Icon = FSlateIcon(FAppStyle::GetAppStyleSetName(), "ClassIcon.TextureCube");
+					const FToolMenuExecuteAction UIAction = FToolMenuExecuteAction::CreateStatic(&ExecuteCreateTextureArray);
+
+					InSection.AddMenuEntry("Texture_TextureCubeArray", Label, ToolTip, Icon, UIAction);
+				}
+			}));
+		}));
+	});
+}
+
+#undef LOCTEXT_NAMESPACE
diff --git a/Engine/Plugins/Editor/EngineAssetDefinitions/Source/Private/AssetDefinition_TextureCube.h b/Engine/Plugins/Editor/EngineAssetDefinitions/Source/Private/AssetDefinition_TextureCube.h
new file mode 100644
--- /dev/null
+++ b/Engine/Plugins/Editor/EngineAssetDefinitions/Source/Private/AssetDefinition_TextureCube.h
@@ -0,0 +1,22 @@
+// Copyright Epic Games, Inc. All Rights Reserved.
+
+#pragma once
+
+#include "CoreMinimal.h"
+#include "Engine/TextureCube.h"
+#include "AssetDefinition_Texture.h"
+
+#include "AssetDefinition_TextureCube.generated.h"
+
+UCLASS()
+class ENGINEASSETDEFINITIONS_API UAssetDefinition_TextureCube : public UAssetDefinition_Texture
+{
+	GENERATED_BODY()
+
+public:
+	// UAssetDefinition Begin
+	virtual FText GetAssetDisplayName() const override { return NSLOCTEXT("AssetTypeActions", "AssetTypeActions_TextureCube", "Texture Cube"); }
+	virtual FLinearColor GetAssetColor() const override { return FLinearColor(FColor(192,64,128)); }
+	virtual TSoftClassPtr<UObject> GetAssetClass() const override { return UTextureCube::StaticClass(); }
+	// UAssetDefinition End
+};
diff --git a/Engine/Plugins/Editor/EngineAssetDefinitions/Source/Private/AssetDefinition_TextureCubeArray.h b/Engine/Plugins/Editor/EngineAssetDefinitions/Source/Private/AssetDefinition_TextureCubeArray.h
new file mode 100644
--- /dev/null
+++ b/Engine/Plugins/Editor/EngineAssetDefinitions/Source/Private/AssetDefinition_TextureCubeArray.h
@@ -0,0 +1,22 @@
+// Copyright Epic Games, Inc. All Rights Reserved.
+
+#pragma once
+
+#include "CoreMinimal.h"
+#include "Engine/TextureCubeArray.h"
+#include "AssetDefinition_Texture.h"
+
+#include "AssetDefinition_TextureCubeArray.generated.h"
+
+UCLASS()
+class ENGINEASSETDEFINITIONS_API UAssetDefinition_TextureCubeArray : public UAssetDefinition_Texture
+{
+	GENERATED_BODY()
+
+public:
+	// UAssetDefinition Begin
+	virtual FText GetAssetDisplayName() const override { return NSLOCTEXT("AssetTypeActions", "AssetTypeActions_TextureCubeArray", "Texture Cube Array"); }
+	virtual FLinearColor GetAssetColor() const override { return FLinearColor(FColor(192, 192, 128)); }
+	virtual TSoftClassPtr<UObject> GetAssetClass() const override { return UTextureCubeArray::StaticClass(); }
+	// UAssetDefinition End
+};
diff --git a/Engine/Plugins/Editor/EngineAssetDefinitions/Source/Private/AssetDefinition_TextureLightProfile.h b/Engine/Plugins/Editor/EngineAssetDefinitions/Source/Private/AssetDefinition_TextureLightProfile.h
new file mode 100644
--- /dev/null
+++ b/Engine/Plugins/Editor/EngineAssetDefinitions/Source/Private/AssetDefinition_TextureLightProfile.h
@@ -0,0 +1,22 @@
+// Copyright Epic Games, Inc. All Rights Reserved.
+
+#pragma once
+
+#include "CoreMinimal.h"
+#include "Engine/TextureLightProfile.h"
+#include "AssetDefinition_Texture.h"
+
+#include "AssetDefinition_TextureLightProfile.generated.h"
+
+UCLASS()
+class ENGINEASSETDEFINITIONS_API UAssetDefinition_TextureLightProfile : public UAssetDefinition_Texture
+{
+	GENERATED_BODY()
+
+public:
+	// UAssetDefinition Begin
+	virtual FText GetAssetDisplayName() const override { return NSLOCTEXT("AssetTypeActions", "AssetTypeActions_TextureLightProfile", "Texture Light Profile"); }
+	virtual FLinearColor GetAssetColor() const override { return FLinearColor(FColor(192,64,192)); }
+	virtual TSoftClassPtr<UObject> GetAssetClass() const override { return UTextureLightProfile::StaticClass(); }
+	// UAssetDefinition End
+};
diff --git a/Engine/Plugins/Editor/EngineAssetDefinitions/Source/Private/AssetDefinition_TextureRenderTarget.cpp b/Engine/Plugins/Editor/EngineAssetDefinitions/Source/Private/AssetDefinition_TextureRenderTarget.cpp
new file mode 100644
--- /dev/null
+++ b/Engine/Plugins/Editor/EngineAssetDefinitions/Source/Private/AssetDefinition_TextureRenderTarget.cpp
@@ -0,0 +1,92 @@
+// Copyright Epic Games, Inc. All Rights Reserved.
+
+#include "AssetDefinition_TextureRenderTarget.h"
+
+#include "AssetToolsModule.h"
+#include "ContentBrowserMenuContexts.h"
+#include "ToolMenus.h"
+#include "Styling/AppStyle.h"
+#include "Engine/Texture2D.h"
+#include "Engine/Texture2DArray.h"
+#include "Engine/TextureCube.h"
+#include "Engine/VolumeTexture.h"
+#include "Engine/TextureRenderTarget2D.h"
+#include "Engine/TextureRenderTarget2DArray.h"
+#include "Engine/TextureRenderTargetCube.h"
+#include "Engine/TextureRenderTargetVolume.h"
+#include "AssetRegistry/AssetRegistryModule.h"
+
+#define LOCTEXT_NAMESPACE "UAssetDefinition_TextureRenderTarget"
+
+namespace MenuExtension_TextureRenderTarget
+{
+	static void ExecuteCreateStaticTexture(const FToolMenuContext& MenuContext)
+	{
+		if (const UContentBrowserAssetContextMenuContext* Context = UContentBrowserAssetContextMenuContext::FindContextWithAssets(MenuContext))
+		{
+			TArray<UTextureRenderTarget*> SelectedRenderTargets = Context->LoadSelectedObjects<UTextureRenderTarget>();
+			for (UTextureRenderTarget* RenderTarget : SelectedRenderTargets)
+			{
+				FString Name;
+				FString PackageName;
+				IAssetTools::Get().CreateUniqueAssetName(RenderTarget->GetOutermost()->GetName(), TEXT("_Tex"), PackageName, Name);
+
+				UObject* NewObj = nullptr;
+				UTextureRenderTarget2D* TexRT = Cast<UTextureRenderTarget2D>(RenderTarget);
+				UTextureRenderTarget2DArray* TexRT2DArray = Cast<UTextureRenderTarget2DArray>(RenderTarget);
+				UTextureRenderTargetCube* TexRTCube = Cast<UTextureRenderTargetCube>(RenderTarget);
+				UTextureRenderTargetVolume* TexRTVolume = Cast<UTextureRenderTargetVolume>(RenderTarget);
+				if( TexRTCube )
+				{
+					// create a static cube texture as well as its 6 faces
+					NewObj = TexRTCube->ConstructTextureCube( CreatePackage(*PackageName), Name, RenderTarget->GetMaskedFlags() );
+				}
+				else if (TexRTVolume)
+				{
+					NewObj = TexRTVolume->ConstructTextureVolume(CreatePackage( *PackageName), Name, RenderTarget->GetMaskedFlags());
+				}
+				else if (TexRT2DArray)
+				{
+					NewObj = TexRT2DArray->ConstructTexture2DArray(CreatePackage(*PackageName), Name, RenderTarget->GetMaskedFlags());
+				}
+				else if( TexRT )
+				{
+					// create a static 2d texture
+					NewObj = TexRT->ConstructTexture2D( CreatePackage(*PackageName), Name, RenderTarget->GetMaskedFlags(), CTF_Default, nullptr );
+				}
+
+				if( NewObj )
+				{
+					// package needs saving
+					NewObj->MarkPackageDirty();
+
+					// Notify the asset registry
+					FAssetRegistryModule::AssetCreated(NewObj);
+				}
+			}
+		}
+	}
+	
+	static FDelayedAutoRegisterHelper DelayedAutoRegister(EDelayedRegisterRunPhase::EndOfEngineInit, []{ 
+		UToolMenus::RegisterStartupCallback(FSimpleMulticastDelegate::FDelegate::CreateLambda([]()
+		{
+			FToolMenuOwnerScoped OwnerScoped(UE_MODULE_NAME);
+			UToolMenu* Menu = UE::ContentBrowser::ExtendToolMenu_AssetContextMenu(UTextureRenderTarget::StaticClass());
+	        
+			FToolMenuSection& Section = Menu->FindOrAddSection("GetAssetActions");
+			Section.AddDynamicEntry("GetAssetActions_UTextureRenderTarget", FNewToolMenuSectionDelegate::CreateLambda([](FToolMenuSection& InSection)
+			{
+				{
+					const TAttribute<FText> Label = LOCTEXT("TextureRenderTarget_CreateStatic", "Create Static Texture");
+					const TAttribute<FText> ToolTip = LOCTEXT("TextureRenderTarget_CreateStaticTooltip", "Creates a static texture from the selected render targets.");
+					const FSlateIcon Icon = FSlateIcon(FAppStyle::GetAppStyleSetName(), "ClassIcon.Texture2D");
+					const FToolMenuExecuteAction UIAction = FToolMenuExecuteAction::CreateStatic(&ExecuteCreateStaticTexture);
+
+					InSection.AddMenuEntry("TextureRenderTarget_CreateStatic", Label, ToolTip, Icon, UIAction);
+				}
+			}));
+		}));
+	});
+}
+
+#undef LOCTEXT_NAMESPACE
diff --git a/Engine/Plugins/Editor/EngineAssetDefinitions/Source/Private/AssetDefinition_TextureRenderTarget.h b/Engine/Plugins/Editor/EngineAssetDefinitions/Source/Private/AssetDefinition_TextureRenderTarget.h
new file mode 100644
--- /dev/null
+++ b/Engine/Plugins/Editor/EngineAssetDefinitions/Source/Private/AssetDefinition_TextureRenderTarget.h
@@ -0,0 +1,23 @@
+// Copyright Epic Games, Inc. All Rights Reserved.
+
+#pragma once
+
+#include "CoreMinimal.h"
+#include "Engine/TextureRenderTarget.h"
+#include "AssetDefinition_Texture.h"
+
+#include "AssetDefinition_TextureRenderTarget.generated.h"
+
+UCLASS()
+class ENGINEASSETDEFINITIONS_API UAssetDefinition_TextureRenderTarget : public UAssetDefinition_Texture
+{
+	GENERATED_BODY()
+
+public:
+	// UAssetDefinition Begin
+	virtual FText GetAssetDisplayName() const override { return NSLOCTEXT("AssetTypeActions", "AssetTypeActions_TextureRenderTarget", "Texture Render Target"); }
+	virtual FLinearColor GetAssetColor() const override { return FLinearColor(FColor(128,64,64)); }
+	virtual TSoftClassPtr<UObject> GetAssetClass() const override { return UTextureRenderTarget::StaticClass(); }
+	virtual bool CanImport() const override { return false; }
+	// UAssetDefinition End
+};
diff --git a/Engine/Plugins/Editor/EngineAssetDefinitions/Source/Private/AssetDefinition_TextureRenderTarget2D.h b/Engine/Plugins/Editor/EngineAssetDefinitions/Source/Private/AssetDefinition_TextureRenderTarget2D.h
new file mode 100644
--- /dev/null
+++ b/Engine/Plugins/Editor/EngineAssetDefinitions/Source/Private/AssetDefinition_TextureRenderTarget2D.h
@@ -0,0 +1,21 @@
+// Copyright Epic Games, Inc. All Rights Reserved.
+
+#pragma once
+
+#include "CoreMinimal.h"
+#include "Engine/TextureRenderTarget2D.h"
+#include "AssetDefinition_TextureRenderTarget.h"
+
+#include "AssetDefinition_TextureRenderTarget2D.generated.h"
+
+UCLASS()
+class ENGINEASSETDEFINITIONS_API UAssetDefinition_TextureRenderTarget2D : public UAssetDefinition_TextureRenderTarget
+{
+	GENERATED_BODY()
+
+public:
+	// UAssetDefinition Begin
+	virtual FText GetAssetDisplayName() const override { return NSLOCTEXT("AssetTypeActions", "AssetTypeActions_TextureRenderTarget2D", "Render Target"); }
+	virtual TSoftClassPtr<UObject> GetAssetClass() const override { return UTextureRenderTarget2D::StaticClass(); }
+	// UAssetDefinition End
+};
diff --git a/Engine/Plugins/Editor/EngineAssetDefinitions/Source/Private/AssetDefinition_TextureRenderTarget2DArray.h b/Engine/Plugins/Editor/EngineAssetDefinitions/Source/Private/AssetDefinition_TextureRenderTarget2DArray.h
new file mode 100644
--- /dev/null
+++ b/Engine/Plugins/Editor/EngineAssetDefinitions/Source/Private/AssetDefinition_TextureRenderTarget2DArray.h
@@ -0,0 +1,21 @@
+// Copyright Epic Games, Inc. All Rights Reserved.
+
+#pragma once
+
+#include "CoreMinimal.h"
+#include "Engine/TextureRenderTarget2DArray.h"
+#include "AssetDefinition_TextureRenderTarget.h"
+
+#include "AssetDefinition_TextureRenderTarget2DArray.generated.h"
+
+UCLASS()
+class ENGINEASSETDEFINITIONS_API UAssetDefinition_TextureRenderTarget2DArray : public UAssetDefinition_TextureRenderTarget
+{
+	GENERATED_BODY()
+
+public:
+	// UAssetDefinition Begin
+	virtual FText GetAssetDisplayName() const override { return NSLOCTEXT("AssetTypeActions", "AssetTypeActions_TextureRenderTarget2DArray", "2D Array Render Target"); }
+	virtual TSoftClassPtr<UObject> GetAssetClass() const override { return UTextureRenderTarget2DArray::StaticClass(); }
+	// UAssetDefinition End
+};
diff --git a/Engine/Plugins/Editor/EngineAssetDefinitions/Source/Private/AssetDefinition_TextureRenderTargetCube.h b/Engine/Plugins/Editor/EngineAssetDefinitions/Source/Private/AssetDefinition_TextureRenderTargetCube.h
new file mode 100644
--- /dev/null
+++ b/Engine/Plugins/Editor/EngineAssetDefinitions/Source/Private/AssetDefinition_TextureRenderTargetCube.h
@@ -0,0 +1,21 @@
+// Copyright Epic Games, Inc. All Rights Reserved.
+
+#pragma once
+
+#include "CoreMinimal.h"
+#include "Engine/TextureRenderTargetCube.h"
+#include "AssetDefinition_TextureRenderTarget.h"
+
+#include "AssetDefinition_TextureRenderTargetCube.generated.h"
+
+UCLASS()
+class ENGINEASSETDEFINITIONS_API UAssetDefinition_TextureRenderTargetCube : public UAssetDefinition_TextureRenderTarget
+{
+	GENERATED_BODY()
+
+public:
+	// UAssetDefinition Begin
+	virtual FText GetAssetDisplayName() const override { return NSLOCTEXT("AssetTypeActions", "AssetTypeActions_TextureRenderTargetCube", "Cube Render Target"); }
+	virtual TSoftClassPtr<UObject> GetAssetClass() const override { return UTextureRenderTargetCube::StaticClass(); }
+	// UAssetDefinition End
+};
diff --git a/Engine/Plugins/Editor/EngineAssetDefinitions/Source/Private/AssetDefinition_TextureRenderTargetVolume.h b/Engine/Plugins/Editor/EngineAssetDefinitions/Source/Private/AssetDefinition_TextureRenderTargetVolume.h
new file mode 100644
--- /dev/null
+++ b/Engine/Plugins/Editor/EngineAssetDefinitions/Source/Private/AssetDefinition_TextureRenderTargetVolume.h
@@ -0,0 +1,21 @@
+// Copyright Epic Games, Inc. All Rights Reserved.
+
+#pragma once
+
+#include "CoreMinimal.h"
+#include "Engine/TextureRenderTargetVolume.h"
+#include "AssetDefinition_TextureRenderTarget.h"
+
+#include "AssetDefinition_TextureRenderTargetVolume.generated.h"
+
+UCLASS()
+class ENGINEASSETDEFINITIONS_API UAssetDefinition_TextureRenderTargetVolume : public UAssetDefinition_TextureRenderTarget
+{
+	GENERATED_BODY()
+
+public:
+	// UAssetDefinition Begin
+	virtual FText GetAssetDisplayName() const override { return NSLOCTEXT("AssetTypeActions", "AssetTypeActions_TextureRenderTargetVolume", "Volume Render Target"); }
+	virtual TSoftClassPtr<UObject> GetAssetClass() const override { return UTextureRenderTargetVolume::StaticClass(); }
+	// UAssetDefinition End
+};
diff --git a/Engine/Plugins/Editor/EngineAssetDefinitions/Source/Private/AssetDefinition_VolumeTexture.h b/Engine/Plugins/Editor/EngineAssetDefinitions/Source/Private/AssetDefinition_VolumeTexture.h
new file mode 100644
--- /dev/null
+++ b/Engine/Plugins/Editor/EngineAssetDefinitions/Source/Private/AssetDefinition_VolumeTexture.h
@@ -0,0 +1,22 @@
+// Copyright Epic Games, Inc. All Rights Reserved.
+
+#pragma once
+
+#include "CoreMinimal.h"
+#include "Engine/VolumeTexture.h"
+#include "AssetDefinition_TextureRenderTarget.h"
+
+#include "AssetDefinition_VolumeTexture.generated.h"
+
+UCLASS()
+class ENGINEASSETDEFINITIONS_API UAssetDefinition_VolumeTexture : public UAssetDefinition_Texture
+{
+	GENERATED_BODY()
+
+public:
+	// UAssetDefinition Begin
+	virtual FText GetAssetDisplayName() const override { return NSLOCTEXT("AssetTypeActions", "AssetTypeActions_VolumeTexture", "Volume Texture"); }
+	virtual FLinearColor GetAssetColor() const override { return FLinearColor(FColor(192, 128, 64)); }
+	virtual TSoftClassPtr<UObject> GetAssetClass() const override { return UVolumeTexture::StaticClass(); }
+	// UAssetDefinition End
+};
diff --git a/Engine/Plugins/Editor/EngineAssetDefinitions/Source/Public/AssetDefinition_Blueprint.h b/Engine/Plugins/Editor/EngineAssetDefinitions/Source/Public/AssetDefinition_Blueprint.h
--- a/Engine/Plugins/Editor/EngineAssetDefinitions/Source/Public/AssetDefinition_Blueprint.h
+++ b/Engine/Plugins/Editor/EngineAssetDefinitions/Source/Public/AssetDefinition_Blueprint.h
@@ -22,7 +22,12 @@ public:
 	virtual FText GetAssetDisplayName() const override { return NSLOCTEXT("AssetTypeActions", "AssetTypeActions_Blueprint", "Blueprint Class"); }
 	virtual FLinearColor GetAssetColor() const override { return FLinearColor(FColor( 63, 126, 255 )); }
 	virtual TSoftClassPtr<UObject> GetAssetClass() const override { return UBlueprint::StaticClass(); }
-
+	virtual TConstArrayView<FAssetCategoryPath> GetAssetCategories() const override
+	{
+		static const auto Categories = { EAssetCategoryPaths::Basic };
+		return Categories;
+	}
+	
 public:
 	virtual UFactory* GetFactoryForBlueprintType(UBlueprint* InBlueprint) const;
 };
diff --git a/Engine/Plugins/Editor/EngineAssetDefinitions/Source/Public/AssetDefinition_World.h b/Engine/Plugins/Editor/EngineAssetDefinitions/Source/Public/AssetDefinition_World.h
--- a/Engine/Plugins/Editor/EngineAssetDefinitions/Source/Public/AssetDefinition_World.h
+++ b/Engine/Plugins/Editor/EngineAssetDefinitions/Source/Public/AssetDefinition_World.h
@@ -18,7 +18,11 @@ public:
 	virtual FText GetAssetDisplayName() const override { return NSLOCTEXT("AssetTypeActions", "AssetTypeActions_World", "Level"); }
 	virtual FLinearColor GetAssetColor() const override { return FAppStyle::Get().GetColor("LevelEditor.AssetColor"); }
 	virtual TSoftClassPtr<UObject> GetAssetClass() const override { return UWorld::StaticClass(); }
-	virtual TConstArrayView<FAssetCategoryPath> GetAssetCategories() const override { return { EAssetCategoryPaths::Basic }; }
+	virtual TConstArrayView<FAssetCategoryPath> GetAssetCategories() const override
+	{
+		static const auto Categories = { EAssetCategoryPaths::Basic };
+		return Categories;
+	}
 
 	virtual TArray<FAssetData> PrepareToActivateAssets(const FAssetActivateArgs& ActivateArgs) const override;
 	virtual EAssetCommandResult OpenAssets(const FAssetOpenArgs& OpenArgs) const override;
diff --git a/Engine/Source/Developer/AssetTools/AssetTools.Build.cs b/Engine/Source/Developer/AssetTools/AssetTools.Build.cs
--- a/Engine/Source/Developer/AssetTools/AssetTools.Build.cs
+++ b/Engine/Source/Developer/AssetTools/AssetTools.Build.cs
@@ -46,7 +46,7 @@ public class AssetTools : ModuleRules
 				"InterchangeCore",
 				"InterchangeEngine",
 				"PhysicsUtilities",
-				"AssetRegistry",
+				"AssetRegistry"
 			}
 		);
 
@@ -81,7 +81,8 @@ public class AssetTools : ModuleRules
 				"Cascade",
 				"PhysicsAssetEditor",
 				"CurveTableEditor",
-				"StaticMeshEditor"
+				"StaticMeshEditor",
+				"VirtualTexturingEditor"
 			}
 		);
 
@@ -105,7 +106,8 @@ public class AssetTools : ModuleRules
                 "SkeletalMeshEditor",
                 "AnimationEditor",
                 "AnimationBlueprintEditor",
-                "AnimationModifiers"
+                "AnimationModifiers",
+                "VirtualTexturingEditor"
             }
 		);
 	}
diff --git a/Engine/Source/Developer/AssetTools/Private/AssetDefinition_AssetTypeActionsProxy.cpp b/Engine/Source/Developer/AssetTools/Private/AssetDefinition_AssetTypeActionsProxy.cpp
--- a/Engine/Source/Developer/AssetTools/Private/AssetDefinition_AssetTypeActionsProxy.cpp
+++ b/Engine/Source/Developer/AssetTools/Private/AssetDefinition_AssetTypeActionsProxy.cpp
@@ -2,6 +2,7 @@
 
 #include "AssetDefinition_AssetTypeActionsProxy.h"
 
+#include "IAssetTools.h"
 #include "Misc/AssetFilterData.h"
 
 void UAssetDefinition_AssetTypeActionsProxy::Initialize(const TSharedRef<IAssetTypeActions>& NewActions)
@@ -14,11 +15,49 @@ FText UAssetDefinition_AssetTypeActionsProxy::GetAssetDisplayName() const
 	return AssetType->GetName();
 }
 
+FText UAssetDefinition_AssetTypeActionsProxy::GetAssetDisplayName(const FAssetData& AssetData) const
+{
+	return AssetType->GetDisplayNameFromAssetData(AssetData);
+}
+
 FText UAssetDefinition_AssetTypeActionsProxy::GetAssetDescription(const FAssetData& AssetData) const
 {
 	return AssetType->GetAssetDescription(AssetData);
 }
 
+TConstArrayView<FAssetCategoryPath> UAssetDefinition_AssetTypeActionsProxy::GetAssetCategories() const
+{
+	if (bAreAssetCategoriesInitialized)
+	{
+		return AssetCategories;
+	}
+
+	bAreAssetCategoriesInitialized = true;
+
+	IAssetTools& AssetTools = IAssetTools::Get();
+	
+	TArray<FAdvancedAssetCategory> AdvancedAssetCategories;
+    AssetTools.GetAllAdvancedAssetCategories(AdvancedAssetCategories);
+	
+	const uint32 CategoryBits = AssetType->GetCategories();
+	for (const FAdvancedAssetCategory& Category : AdvancedAssetCategories)
+	{
+		if (Category.CategoryType & CategoryBits)
+		{
+			TArray<FText> CategoryTextArray;
+			CategoryTextArray.Add(Category.CategoryName);
+			for (const FText& Submenu : AssetType->GetSubMenus())
+			{
+				CategoryTextArray.Add(Submenu);
+			}
+		
+			AssetCategories.Add(FAssetCategoryPath(CategoryTextArray));
+		}
+	}
+
+	return AssetCategories;
+}
+
 TSoftClassPtr<UObject> UAssetDefinition_AssetTypeActionsProxy::GetAssetClass() const
 {
 	if (UClass* SupportedClass = AssetType->GetSupportedClass())
@@ -34,6 +73,32 @@ FLinearColor UAssetDefinition_AssetTypeActionsProxy::GetAssetColor() const
 	return FLinearColor(AssetType->GetTypeColor());
 }
 
+FAssetSupportResponse UAssetDefinition_AssetTypeActionsProxy::CanRename(const FAssetData& InAsset) const
+{
+    FText OutErrorMsg;
+	if (AssetType->CanRename(InAsset, &OutErrorMsg))
+	{
+	     return FAssetSupportResponse::Supported();
+	}
+	else
+	{
+	    return OutErrorMsg.IsEmpty() ? FAssetSupportResponse::NotSupported() : FAssetSupportResponse::Error(OutErrorMsg);
+	}
+}
+
+FAssetSupportResponse UAssetDefinition_AssetTypeActionsProxy::CanDuplicate(const FAssetData& InAsset) const
+{
+    FText OutErrorMsg;
+	if (AssetType->CanDuplicate(InAsset, &OutErrorMsg))
+	{
+	     return FAssetSupportResponse::Supported();
+	}
+	else
+	{
+	    return OutErrorMsg.IsEmpty() ? FAssetSupportResponse::NotSupported() : FAssetSupportResponse::Error(OutErrorMsg);
+	}
+}
+
 FAssetSupportResponse UAssetDefinition_AssetTypeActionsProxy::CanLocalize(const FAssetData& InAsset) const
 {
 	return AssetType->CanLocalize() ? FAssetSupportResponse::Supported() : FAssetSupportResponse::NotSupported();
@@ -49,6 +114,46 @@ bool UAssetDefinition_AssetTypeActionsProxy::CanMerge() const
 	return AssetType->CanMerge();
 }
 
+EAssetCommandResult UAssetDefinition_AssetTypeActionsProxy::Merge(const FAssetMergeArgs& MergeArgs) const
+{
+	if (!MergeArgs.BaseAsset.IsSet() && !MergeArgs.RemoteAsset.IsSet())
+	{
+		AssetType->Merge(MergeArgs.LocalAsset.GetAsset());
+		return EAssetCommandResult::Handled;
+	}
+	else
+	{
+		FOnAssetMergeResolved MergeCallback = MergeArgs.ResolutionCallback;
+		AssetType->Merge(MergeArgs.BaseAsset.GetValue().GetAsset(), MergeArgs.RemoteAsset.GetValue().GetAsset(), MergeArgs.LocalAsset.GetAsset(),
+			FOnMergeResolved::CreateLambda([MergeCallback](UPackage* MergedPackage, EMergeResult::Type Result){
+			FAssetMergeResults Results;
+			Results.MergedPackage = MergedPackage;
+			Results.Result = static_cast<EAssetMergeResult>(Result);
+			MergeCallback.ExecuteIfBound(Results);
+		}));
+		return EAssetCommandResult::Handled;
+	}
+}
+
+FAssetOpenSupport UAssetDefinition_AssetTypeActionsProxy::GetAssetOpenSupport(const FAssetOpenSupportArgs& OpenSupportArgs) const
+{
+	const bool bIsSupported = AssetType->SupportsOpenedMethod(static_cast<EAssetTypeActivationOpenedMethod>(OpenSupportArgs.OpenMethod));
+
+	FAssetOpenSupport Support(OpenSupportArgs.OpenMethod, bIsSupported);
+
+	if (AssetType->ShouldForceWorldCentric())
+	{
+		Support.RequiredToolkitMode = EToolkitMode::WorldCentric;
+	}
+	
+	return Support;
+}
+
+TArray<FAssetData> UAssetDefinition_AssetTypeActionsProxy::PrepareToActivateAssets(const FAssetActivateArgs& ActivateArgs) const
+{
+	return AssetType->GetValidAssetsForPreviewOrEdit(ActivateArgs.Assets, ActivateArgs.ActivationMethod == EAssetActivationMethod::Previewed);
+}
+
 EAssetCommandResult UAssetDefinition_AssetTypeActionsProxy::ActivateAssets(const FAssetActivateArgs& ActivateArgs) const
 {
 	EAssetTypeActivationMethod::Type ActivationType = EAssetTypeActivationMethod::Opened;
@@ -70,6 +175,12 @@ EAssetCommandResult UAssetDefinition_AssetTypeActionsProxy::ActivateAssets(const
 	return bResult ? EAssetCommandResult::Handled : EAssetCommandResult::Unhandled;
 }
 
+EAssetCommandResult UAssetDefinition_AssetTypeActionsProxy::OpenAssets(const FAssetOpenArgs& OpenArgs) const
+{
+	AssetType->OpenAssetEditor(OpenArgs.LoadObjects<UObject>(), static_cast<EAssetTypeActivationOpenedMethod>(OpenArgs.OpenMethod), OpenArgs.ToolkitHost);
+	return EAssetCommandResult::Handled;
+}
+
 EAssetCommandResult UAssetDefinition_AssetTypeActionsProxy::GetFilters(TArray<FAssetFilterData>& OutFilters) const
 {
 	if (AssetType->CanFilter())
@@ -86,3 +197,33 @@ EAssetCommandResult UAssetDefinition_AssetTypeActionsProxy::GetFilters(TArray<FA
 	return EAssetCommandResult::Unhandled;
 }
 
+FText UAssetDefinition_AssetTypeActionsProxy::GetObjectDisplayText(UObject* Object) const
+{
+	return FText::FromString(AssetType->GetObjectDisplayName(Object));
+}
+
+EAssetCommandResult UAssetDefinition_AssetTypeActionsProxy::PerformAssetDiff(const FAssetDiffArgs& DiffArgs) const
+{
+	AssetType->PerformAssetDiff(DiffArgs.OldAsset, DiffArgs.NewAsset, DiffArgs.OldRevision, DiffArgs.NewRevision);
+	return EAssetCommandResult::Handled;
+}
+
+UThumbnailInfo* UAssetDefinition_AssetTypeActionsProxy::LoadThumbnailInfo(const FAssetData& Asset) const
+{
+	return AssetType->GetThumbnailInfo(Asset.GetAsset());
+}
+
+const FSlateBrush* UAssetDefinition_AssetTypeActionsProxy::GetThumbnailBrush(const FAssetData& InAssetData, const FName InClassName) const
+{
+	return AssetType->GetThumbnailBrush(InAssetData, InClassName);
+}
+
+const FSlateBrush* UAssetDefinition_AssetTypeActionsProxy::GetIconBrush(const FAssetData& InAssetData, const FName InClassName) const
+{
+	return AssetType->GetIconBrush(InAssetData, InClassName);
+}
+
+TSharedPtr<SWidget> UAssetDefinition_AssetTypeActionsProxy::GetThumbnailOverlay(const FAssetData& InAssetData) const
+{
+	return AssetType->GetThumbnailOverlay(InAssetData);
+}
\ No newline at end of file
diff --git a/Engine/Source/Developer/AssetTools/Private/AssetDefinition_AssetTypeActionsProxy.h b/Engine/Source/Developer/AssetTools/Private/AssetDefinition_AssetTypeActionsProxy.h
--- a/Engine/Source/Developer/AssetTools/Private/AssetDefinition_AssetTypeActionsProxy.h
+++ b/Engine/Source/Developer/AssetTools/Private/AssetDefinition_AssetTypeActionsProxy.h
@@ -20,24 +20,43 @@ public:
 	const TSharedPtr<IAssetTypeActions>& GetAssetType() const { return AssetType; }
 
 	virtual FText GetAssetDisplayName() const override;
+	virtual FText GetAssetDisplayName(const FAssetData& InAssetData) const override;
 	virtual TSoftClassPtr<UObject> GetAssetClass() const override;
 	virtual FLinearColor GetAssetColor() const override;
-	virtual FText GetAssetDescription(const FAssetData& AssetData) const override;
+	virtual FText GetAssetDescription(const FAssetData& InAssetData) const override;
+	virtual TConstArrayView<FAssetCategoryPath> GetAssetCategories() const override;
 
+	virtual FAssetSupportResponse CanRename(const FAssetData& InAsset) const override;
+	virtual FAssetSupportResponse CanDuplicate(const FAssetData& InAsset) const override;
 	virtual FAssetSupportResponse CanLocalize(const FAssetData& InAsset) const override;
 	virtual bool CanImport() const override;
 	virtual bool CanMerge() const override;
+	virtual EAssetCommandResult Merge(const FAssetMergeArgs& MergeArgs) const override;
 
+	virtual FAssetOpenSupport GetAssetOpenSupport(const FAssetOpenSupportArgs& OpenSupportArgs) const override;
+	virtual TArray<FAssetData> PrepareToActivateAssets(const FAssetActivateArgs& ActivateArgs) const override;
 	virtual EAssetCommandResult ActivateAssets(const FAssetActivateArgs& ActivateArgs) const override;
+	virtual EAssetCommandResult OpenAssets(const FAssetOpenArgs& OpenArgs) const override;
 
 	virtual EAssetCommandResult GetFilters(TArray<FAssetFilterData>& OutFilters) const override;
 
 	// We don't need to implement GetSourceFiles, the default one will / should work for any existing type.
 	// virtual EAssetCommandResult GetSourceFiles(const FAssetSourceFileArgs& SourceFileArgs, TArray<FAssetSourceFile>& OutSourceAssets) const override;
+
+	virtual FText GetObjectDisplayText(UObject* Object) const override;
+	virtual EAssetCommandResult PerformAssetDiff(const FAssetDiffArgs& DiffArgs) const override;
 	
+	virtual UThumbnailInfo* LoadThumbnailInfo(const FAssetData& Asset) const override;
+	virtual const FSlateBrush* GetThumbnailBrush(const FAssetData& InAssetData, const FName InClassName) const override;
+	virtual const FSlateBrush* GetIconBrush(const FAssetData& InAssetData, const FName InClassName) const override;
+	virtual TSharedPtr<SWidget> GetThumbnailOverlay(const FAssetData& InAssetData) const override;
+
 protected:
 	virtual bool CanRegisterStatically() const override { return false; }
 	
 private:
 	TSharedPtr<IAssetTypeActions> AssetType;
+	
+	mutable bool bAreAssetCategoriesInitialized = false;
+	mutable TArray<FAssetCategoryPath> AssetCategories;
 };
diff --git a/Engine/Source/Developer/AssetTools/Private/AssetTools.cpp b/Engine/Source/Developer/AssetTools/Private/AssetTools.cpp
--- a/Engine/Source/Developer/AssetTools/Private/AssetTools.cpp
+++ b/Engine/Source/Developer/AssetTools/Private/AssetTools.cpp
@@ -43,8 +43,6 @@
 #include "AssetTypeActions/AssetTypeActions_MaterialInterface.h"
 #include "AssetTypeActions/AssetTypeActions_SkeletalMesh.h"
 #include "AssetTypeActions/AssetTypeActions_FbxSceneImportData.h"
-#include "AssetTypeActions/AssetTypeActions_Texture.h"
-#include "AssetTypeActions/AssetTypeActions_TextureRenderTarget.h"
 #include "AssetTypeActions/AssetTypeActions_VectorField.h"
 #include "AssetTypeActions/AssetTypeActions_AnimationAsset.h"
 #include "AssetTypeActions/AssetTypeActions_AnimBlueprint.h"
@@ -60,14 +58,11 @@
 #include "AssetTypeActions/AssetTypeActions_AimOffset.h"
 #include "AssetTypeActions/AssetTypeActions_BlendSpace1D.h"
 #include "AssetTypeActions/AssetTypeActions_AimOffset1D.h"
-#include "AssetTypeActions/AssetTypeActions_TextureRenderTarget2D.h"
-#include "AssetTypeActions/AssetTypeActions_CanvasRenderTarget2D.h"
 #include "AssetTypeActions/AssetTypeActions_CurveFloat.h"
 #include "AssetTypeActions/AssetTypeActions_CurveTable.h"
 #include "AssetTypeActions/AssetTypeActions_CompositeCurveTable.h"
 #include "AssetTypeActions/AssetTypeActions_CurveVector.h"
 #include "AssetTypeActions/AssetTypeActions_CurveLinearColor.h"
-#include "AssetTypeActions/AssetTypeActions_CurveLinearColorAtlas.h"
 #include "AssetTypeActions/AssetTypeActions_DataAsset.h"
 #include "AssetTypeActions/AssetTypeActions_DataLayer.h"
 #include "AssetTypeActions/AssetTypeActions_DataTable.h"
@@ -108,15 +103,6 @@
 #include "AssetTypeActions/AssetTypeActions_SlateWidgetStyle.h"
 #include "AssetTypeActions/AssetTypeActions_StaticMesh.h"
 #include "AssetTypeActions/AssetTypeActions_SubUVAnimation.h"
-#include "AssetTypeActions/AssetTypeActions_Texture2D.h"
-#include "AssetTypeActions/AssetTypeActions_Texture2DArray.h"
-#include "AssetTypeActions/AssetTypeActions_TextureCube.h"
-#include "AssetTypeActions/AssetTypeActions_TextureCubeArray.h"
-#include "AssetTypeActions/AssetTypeActions_VolumeTexture.h"
-#include "AssetTypeActions/AssetTypeActions_TextureRenderTarget2DArray.h"
-#include "AssetTypeActions/AssetTypeActions_TextureRenderTargetCube.h"
-#include "AssetTypeActions/AssetTypeActions_TextureRenderTargetVolume.h"
-#include "AssetTypeActions/AssetTypeActions_TextureLightProfile.h"
 #include "AssetTypeActions/AssetTypeActions_TouchInterface.h"
 #include "AssetTypeActions/AssetTypeActions_VectorFieldAnimated.h"
 #include "AssetTypeActions/AssetTypeActions_VectorFieldStatic.h"
@@ -149,7 +135,6 @@
 #include "AdvancedCopyCustomization.h"
 #include "SAdvancedCopyReportDialog.h"
 #include "AssetToolsSettings.h"
-#include "AssetVtConversion.h"
 #include "Misc/ConfigCacheIni.h"
 #include "Misc/NamePermissionList.h"
 #include "InterchangeManager.h"
@@ -180,6 +165,7 @@
 
 #include "AssetDefinition.h"
 #include "AssetDefinitionRegistry.h"
+#include "VirtualTexturingEditorModule.h"
 #include "Algo/AnyOf.h"
 #include "Misc/AssetFilterData.h"
 
@@ -289,7 +275,7 @@ public:
 	virtual TArray<FAssetData> GetValidAssetsForPreviewOrEdit(TArrayView<const FAssetData> InAssetDatas, bool bIsPreview) override
 	{
 		FAssetActivateArgs ActivateArgs;
-		ActivateArgs.ActivationMethod = bIsPreview ? EAssetActivationMethod::Opened : EAssetActivationMethod::Previewed;
+		ActivateArgs.ActivationMethod = bIsPreview ? EAssetActivationMethod::Previewed : EAssetActivationMethod::Opened;
 		ActivateArgs.Assets = InAssetDatas;
 
 		return AssetDefinitionPtr.Get()->PrepareToActivateAssets(ActivateArgs);
@@ -1253,15 +1239,12 @@ UAssetToolsImpl::UAssetToolsImpl(const FObjectInitializer& ObjectInitializer)
 	RegisterAssetTypeActions(MakeShareable(new FAssetTypeActions_BlendSpace1D));
 	RegisterAssetTypeActions(MakeShareable(new FAssetTypeActions_Blueprint));
 	RegisterAssetTypeActions(MakeShareable(new FAssetTypeActions_BlueprintGeneratedClass));
-	RegisterAssetTypeActions(MakeShareable(new FAssetTypeActions_CanvasRenderTarget2D));
 	RegisterAssetTypeActions(MakeShareable(new FAssetTypeActions_Curve));
 	RegisterAssetTypeActions(MakeShareable(new FAssetTypeActions_CurveFloat));
 	RegisterAssetTypeActions(MakeShareable(new FAssetTypeActions_CurveTable));
 	RegisterAssetTypeActions(MakeShareable(new FAssetTypeActions_CompositeCurveTable));
 	RegisterAssetTypeActions(MakeShareable(new FAssetTypeActions_CurveVector));
 	RegisterAssetTypeActions(MakeShareable(new FAssetTypeActions_CurveLinearColor));
-	RegisterAssetTypeActions(MakeShareable(new FAssetTypeActions_CurveLinearColorAtlas));
-	//RegisterAssetTypeActions(MakeShareable(new FAssetTypeActions_DataAsset));
 	RegisterAssetTypeActions(MakeShareable(new FAssetTypeActions_DataLayer));
 	RegisterAssetTypeActions(MakeShareable(new FAssetTypeActions_DataTable));
 	RegisterAssetTypeActions(MakeShareable(new FAssetTypeActions_CompositeDataTable));
@@ -1306,18 +1289,6 @@ UAssetToolsImpl::UAssetToolsImpl(const FObjectInitializer& ObjectInitializer)
 	RegisterAssetTypeActions(MakeShareable(new FAssetTypeActions_SlateBrush));
 	RegisterAssetTypeActions(MakeShareable(new FAssetTypeActions_SlateWidgetStyle));
 	RegisterAssetTypeActions(MakeShareable(new FAssetTypeActions_StaticMesh));
-	RegisterAssetTypeActions(MakeShareable(new FAssetTypeActions_Texture));
-	RegisterAssetTypeActions(MakeShareable(new FAssetTypeActions_Texture2D));
-	RegisterAssetTypeActions(MakeShareable(new FAssetTypeActions_TextureCube));
-	RegisterAssetTypeActions(MakeShareable(new FAssetTypeActions_Texture2DArray));
-	RegisterAssetTypeActions(MakeShareable(new FAssetTypeActions_TextureCubeArray));
-	RegisterAssetTypeActions(MakeShareable(new FAssetTypeActions_VolumeTexture));
-	RegisterAssetTypeActions(MakeShareable(new FAssetTypeActions_TextureRenderTarget));
-	RegisterAssetTypeActions(MakeShareable(new FAssetTypeActions_TextureRenderTarget2D));
-	RegisterAssetTypeActions(MakeShareable(new FAssetTypeActions_TextureRenderTarget2DArray));
-	RegisterAssetTypeActions(MakeShareable(new FAssetTypeActions_TextureRenderTargetCube));
-	RegisterAssetTypeActions(MakeShareable(new FAssetTypeActions_TextureRenderTargetVolume));
-	RegisterAssetTypeActions(MakeShareable(new FAssetTypeActions_TextureLightProfile));
 	RegisterAssetTypeActions(MakeShareable(new FAssetTypeActions_TouchInterface));
 	RegisterAssetTypeActions(MakeShareable(new FAssetTypeActions_VectorField));
 	RegisterAssetTypeActions(MakeShareable(new FAssetTypeActions_VectorFieldAnimated));
@@ -5029,16 +5000,8 @@ void UAssetToolsImpl::OpenEditorForAssets(const TArray<UObject*>& Assets)
 
 void UAssetToolsImpl::ConvertVirtualTextures(const TArray<UTexture2D *>& Textures, bool bConvertBackToNonVirtual, const TArray<UMaterial *>* RelatedMaterials /* = nullptr */) const
 {
-	FVTConversionWorker VirtualTextureConversionWorker(bConvertBackToNonVirtual);
-	VirtualTextureConversionWorker.UserTextures = Textures;
-	//We want all given texture to be added, so we put a minimum texture size of 0
-	VirtualTextureConversionWorker.FilterList(0);
-	if (RelatedMaterials)
-	{
-		VirtualTextureConversionWorker.Materials.Append(*RelatedMaterials);
-	}
-
-	VirtualTextureConversionWorker.DoConvert();
+	IVirtualTexturingEditorModule* Module = FModuleManager::Get().GetModulePtr<IVirtualTexturingEditorModule>("VirtualTexturingEditor");
+	Module->ConvertVirtualTextures(Textures, bConvertBackToNonVirtual, RelatedMaterials);
 }
 
 void UAssetToolsImpl::BeginAdvancedCopyPackages(const TArray<FName>& InputNamesToCopy, const FString& TargetPath) const
diff --git a/Engine/Source/Developer/AssetTools/Private/AssetTypeActions/AssetTypeActions_CanvasRenderTarget2D.h b/Engine/Source/Developer/AssetTools/Private/AssetTypeActions/AssetTypeActions_CanvasRenderTarget2D.h
deleted file mode 100644
diff --git a/Engine/Source/Developer/AssetTools/Private/AssetTypeActions/AssetTypeActions_CurveLinearColorAtlas.cpp b/Engine/Source/Developer/AssetTools/Private/AssetTypeActions/AssetTypeActions_CurveLinearColorAtlas.cpp
deleted file mode 100644
diff --git a/Engine/Source/Developer/AssetTools/Private/AssetTypeActions/AssetTypeActions_CurveLinearColorAtlas.h b/Engine/Source/Developer/AssetTools/Private/AssetTypeActions/AssetTypeActions_CurveLinearColorAtlas.h
deleted file mode 100644
diff --git a/Engine/Source/Developer/AssetTools/Private/AssetTypeActions/AssetTypeActions_Texture.cpp b/Engine/Source/Developer/AssetTools/Private/AssetTypeActions/AssetTypeActions_Texture.cpp
deleted file mode 100644
diff --git a/Engine/Source/Developer/AssetTools/Private/AssetTypeActions/AssetTypeActions_Texture.h b/Engine/Source/Developer/AssetTools/Private/AssetTypeActions/AssetTypeActions_Texture.h
deleted file mode 100644
diff --git a/Engine/Source/Developer/AssetTools/Private/AssetTypeActions/AssetTypeActions_Texture2D.cpp b/Engine/Source/Developer/AssetTools/Private/AssetTypeActions/AssetTypeActions_Texture2D.cpp
deleted file mode 100644
diff --git a/Engine/Source/Developer/AssetTools/Private/AssetTypeActions/AssetTypeActions_Texture2D.h b/Engine/Source/Developer/AssetTools/Private/AssetTypeActions/AssetTypeActions_Texture2D.h
deleted file mode 100644
diff --git a/Engine/Source/Developer/AssetTools/Private/AssetTypeActions/AssetTypeActions_Texture2DArray.cpp b/Engine/Source/Developer/AssetTools/Private/AssetTypeActions/AssetTypeActions_Texture2DArray.cpp
deleted file mode 100644
diff --git a/Engine/Source/Developer/AssetTools/Private/AssetTypeActions/AssetTypeActions_Texture2DArray.h b/Engine/Source/Developer/AssetTools/Private/AssetTypeActions/AssetTypeActions_Texture2DArray.h
deleted file mode 100644
diff --git a/Engine/Source/Developer/AssetTools/Private/AssetTypeActions/AssetTypeActions_TextureCube.cpp b/Engine/Source/Developer/AssetTools/Private/AssetTypeActions/AssetTypeActions_TextureCube.cpp
deleted file mode 100644
diff --git a/Engine/Source/Developer/AssetTools/Private/AssetTypeActions/AssetTypeActions_TextureCube.h b/Engine/Source/Developer/AssetTools/Private/AssetTypeActions/AssetTypeActions_TextureCube.h
deleted file mode 100644
diff --git a/Engine/Source/Developer/AssetTools/Private/AssetTypeActions/AssetTypeActions_TextureCubeArray.cpp b/Engine/Source/Developer/AssetTools/Private/AssetTypeActions/AssetTypeActions_TextureCubeArray.cpp
deleted file mode 100644
diff --git a/Engine/Source/Developer/AssetTools/Private/AssetTypeActions/AssetTypeActions_TextureCubeArray.h b/Engine/Source/Developer/AssetTools/Private/AssetTypeActions/AssetTypeActions_TextureCubeArray.h
deleted file mode 100644
diff --git a/Engine/Source/Developer/AssetTools/Private/AssetTypeActions/AssetTypeActions_TextureLightProfile.cpp b/Engine/Source/Developer/AssetTools/Private/AssetTypeActions/AssetTypeActions_TextureLightProfile.cpp
deleted file mode 100644
diff --git a/Engine/Source/Developer/AssetTools/Private/AssetTypeActions/AssetTypeActions_TextureLightProfile.h b/Engine/Source/Developer/AssetTools/Private/AssetTypeActions/AssetTypeActions_TextureLightProfile.h
deleted file mode 100644
diff --git a/Engine/Source/Developer/AssetTools/Private/AssetTypeActions/AssetTypeActions_TextureRenderTarget.cpp b/Engine/Source/Developer/AssetTools/Private/AssetTypeActions/AssetTypeActions_TextureRenderTarget.cpp
deleted file mode 100644
diff --git a/Engine/Source/Developer/AssetTools/Private/AssetTypeActions/AssetTypeActions_TextureRenderTarget.h b/Engine/Source/Developer/AssetTools/Private/AssetTypeActions/AssetTypeActions_TextureRenderTarget.h
deleted file mode 100644
diff --git a/Engine/Source/Developer/AssetTools/Private/AssetTypeActions/AssetTypeActions_TextureRenderTarget2D.h b/Engine/Source/Developer/AssetTools/Private/AssetTypeActions/AssetTypeActions_TextureRenderTarget2D.h
deleted file mode 100644
diff --git a/Engine/Source/Developer/AssetTools/Private/AssetTypeActions/AssetTypeActions_TextureRenderTarget2DArray.h b/Engine/Source/Developer/AssetTools/Private/AssetTypeActions/AssetTypeActions_TextureRenderTarget2DArray.h
deleted file mode 100644
diff --git a/Engine/Source/Developer/AssetTools/Private/AssetTypeActions/AssetTypeActions_TextureRenderTargetCube.h b/Engine/Source/Developer/AssetTools/Private/AssetTypeActions/AssetTypeActions_TextureRenderTargetCube.h
deleted file mode 100644
diff --git a/Engine/Source/Developer/AssetTools/Private/AssetTypeActions/AssetTypeActions_TextureRenderTargetVolume.h b/Engine/Source/Developer/AssetTools/Private/AssetTypeActions/AssetTypeActions_TextureRenderTargetVolume.h
deleted file mode 100644
diff --git a/Engine/Source/Developer/AssetTools/Private/AssetTypeActions/AssetTypeActions_VolumeTexture.h b/Engine/Source/Developer/AssetTools/Private/AssetTypeActions/AssetTypeActions_VolumeTexture.h
deleted file mode 100644
diff --git a/Engine/Source/Developer/AssetTools/Private/AssetVtConversion.cpp b/Engine/Source/Developer/AssetTools/Private/AssetVtConversion.cpp
deleted file mode 100644
diff --git a/Engine/Source/Developer/AssetTools/Private/AssetVtConversion.h b/Engine/Source/Developer/AssetTools/Private/AssetVtConversion.h
deleted file mode 100644
diff --git a/Engine/Source/Editor/AssetDefinition/Private/AssetDefinition.cpp b/Engine/Source/Editor/AssetDefinition/Private/AssetDefinition.cpp
--- a/Engine/Source/Editor/AssetDefinition/Private/AssetDefinition.cpp
+++ b/Engine/Source/Editor/AssetDefinition/Private/AssetDefinition.cpp
@@ -75,9 +75,4 @@ EAssetCommandResult UAssetDefinition::GetFilters(TArray<FAssetFilterData>& OutFi
 	return EAssetCommandResult::Unhandled;
 }
 
-TConstArrayView<FAssetCategoryPath> UAssetDefinition::GetAssetCategories() const
-{
-	return { EAssetCategoryPaths::Basic };
-}
-
 #undef LOCTEXT_NAMESPACE
\ No newline at end of file
diff --git a/Engine/Source/Editor/AssetDefinition/Public/AssetDefinition.h b/Engine/Source/Editor/AssetDefinition/Public/AssetDefinition.h
--- a/Engine/Source/Editor/AssetDefinition/Public/AssetDefinition.h
+++ b/Engine/Source/Editor/AssetDefinition/Public/AssetDefinition.h
@@ -5,6 +5,7 @@
 #include "UObject/Object.h"
 #include "UObject/SoftObjectPtr.h"
 #include "AssetRegistry/AssetData.h"
+#include "Misc/ScopedSlowTask.h"
 #include "Toolkits/IToolkit.h"
 
 #include "AssetDefinition.generated.h"
@@ -56,11 +57,15 @@ struct FAssetArgs
 	template<typename ExpectedObjectType>
 	TArray<ExpectedObjectType*> LoadObjects(const TSet<FName>& LoadTags = {}) const
 	{
+		FScopedSlowTask SlowTask(Assets.Num());
+	
 		TArray<ExpectedObjectType*> LoadedObjects;
 		LoadedObjects.Reserve(Assets.Num());
 		
 		for (const FAssetData& Asset : Assets)
 		{
+			SlowTask.EnterProgressFrame(1, FText::FromString(Asset.GetObjectPathString()));
+			
 			if (Asset.IsInstanceOf(ExpectedObjectType::StaticClass()))
 			{
 				if (ExpectedObjectType* ExpectedType = Cast<ExpectedObjectType>(Asset.GetAsset(LoadTags)))
@@ -179,9 +184,14 @@ struct FAssetDiffArgs
 
 struct FAssetOpenSupportArgs
 {
-	EAssetOpenMethod OpenMethod;
+	EAssetOpenMethod OpenMethod = EAssetOpenMethod::Edit;
 };
 
+/**
+ * The asset category path is how we know how to build menus around assets.  For example, Basic is generally the ones
+ * we expose at the top level, where as everything else is a category with a pull out menu, and the subcategory would
+ * be where it gets placed in a submenu inside of there.
+ */
 struct ASSETDEFINITION_API FAssetCategoryPath
 {
 	FAssetCategoryPath(FText InCategory);
@@ -193,11 +203,15 @@ struct ASSETDEFINITION_API FAssetCategoryPath
 	bool HasSubCategory() const { return CategoryPath.Num() > 1; }
 	FName GetSubCategory() const { return HasSubCategory() ? CategoryPath[1].Key : NAME_None; }
 	FText GetSubCategoryText() const { return HasSubCategory() ? CategoryPath[1].Value : FText::GetEmpty(); }
-
+	
 private:
 	TArray<TPair<FName, FText>> CategoryPath;
 };
 
+/**
+ * These are just some common asset categories.  You're not at all limited to these, and can register an "Advanced"
+ * category with the IAssetTools::RegisterAdvancedAssetCategory.
+ */
 struct ASSETDEFINITION_API EAssetCategoryPaths
 {
 	static FAssetCategoryPath Basic;
@@ -266,7 +280,11 @@ public:
 	virtual FText GetAssetDescription(const FAssetData& AssetData) const { return FText::GetEmpty(); }
 
 	/** Gets a list of categories this asset is in, these categories are used to help organize */
-	virtual TConstArrayView<FAssetCategoryPath> GetAssetCategories() const;
+	virtual TConstArrayView<FAssetCategoryPath> GetAssetCategories() const
+	{
+		static const auto Categories = { EAssetCategoryPaths::Misc };
+		return Categories;
+	}
 	
 public:
 	// Common Operations
@@ -340,7 +358,7 @@ public:
 	// Thumbnails
 
 	/** Returns the thumbnail info for the specified asset, if it has one. This typically requires loading the asset.  */
-	virtual UThumbnailInfo* LoadThumbnailInfo(const FAssetData& Asset) const
+	virtual UThumbnailInfo* LoadThumbnailInfo(const FAssetData& InAssetData) const
 	{
 		return nullptr;
 	}
@@ -358,7 +376,7 @@ public:
 	}
 	
 	/** Optionally returns a custom widget to overlay on top of this assets' thumbnail */
-	virtual TSharedPtr<SWidget> GetThumbnailOverlay(const FAssetData& AssetData) const
+	virtual TSharedPtr<SWidget> GetThumbnailOverlay(const FAssetData& InAssetData) const
 	{
 		return TSharedPtr<SWidget>();
 	}
diff --git a/Engine/Source/Editor/VirtualTexturingEditor/Private/RuntimeVirtualTextureAssetTypeActions.cpp b/Engine/Source/Editor/VirtualTexturingEditor/Private/RuntimeVirtualTextureAssetTypeActions.cpp
--- a/Engine/Source/Editor/VirtualTexturingEditor/Private/RuntimeVirtualTextureAssetTypeActions.cpp
+++ b/Engine/Source/Editor/VirtualTexturingEditor/Private/RuntimeVirtualTextureAssetTypeActions.cpp
@@ -11,6 +11,7 @@
 #include "AssetRegistry/IAssetRegistry.h"
 #include "IContentBrowserSingleton.h"
 #include "MaterialEditingLibrary.h"
+#include "AssetRegistry/AssetRegistryHelpers.h"
 #include "Materials/Material.h"
 #include "Materials/MaterialExpressionRuntimeVirtualTextureSample.h"
 #include "Materials/MaterialFunction.h"
@@ -27,41 +28,11 @@ DEFINE_LOG_CATEGORY_STATIC(LogRuntimeVirtualTextureFixMaterial, Log, Log);
 
 namespace
 {
-	void GetReferencersData(UObject *Object, UClass *MatchClass, TArray<FAssetData> &OutAssetDatas)
-	{
-		FAssetRegistryModule& AssetRegistryModule = FModuleManager::LoadModuleChecked<FAssetRegistryModule>("AssetRegistry");
-		IAssetRegistry& AssetRegistry = AssetRegistryModule.Get();
-
-		TArray<FAssetIdentifier> Referencers;
-		AssetRegistry.GetReferencers(Object->GetOuter()->GetFName(), Referencers);
-
-		for (auto AssetIdentifier : Referencers)
-		{
-			TArray<FAssetData> Assets;
-			AssetRegistry.GetAssetsByPackageName(AssetIdentifier.PackageName, Assets);
-
-			for (auto AssetData : Assets)
-			{
-				if (MatchClass != nullptr)
-				{
-					if (AssetData.IsInstanceOf(MatchClass))
-					{
-						OutAssetDatas.AddUnique(AssetData);
-					}
-				}
-				else
-				{
-					OutAssetDatas.AddUnique(AssetData);
-				}
-			}
-		}
-	}
-
 	template <class T> 
 	void GetReferencersOfType(UObject *Object, TArray<T*> &OutObjects)
 	{
 		TArray<FAssetData> AssetDatas;
-		GetReferencersData(Object, T::StaticClass(), AssetDatas);
+		UAssetRegistryHelpers::FindReferencersOfAssetOfClass(Object, { T::StaticClass() }, AssetDatas);
 
 		for (auto Data : AssetDatas)
 		{
@@ -253,8 +224,7 @@ void FAssetTypeActions_RuntimeVirtualTexture::ExecuteFindMaterials(TWeakObjectPt
 	URuntimeVirtualTexture* RuntimeVirtualTexture = Object.Get();
 	if (RuntimeVirtualTexture != nullptr)
 	{
-		GetReferencersData(RuntimeVirtualTexture, UMaterialInterface::StaticClass(), Materials);
-		GetReferencersData(RuntimeVirtualTexture, UMaterialFunction::StaticClass(), Materials);
+		UAssetRegistryHelpers::FindReferencersOfAssetOfClass(RuntimeVirtualTexture, { UMaterialInterface::StaticClass(), UMaterialFunction::StaticClass() }, Materials);
 	}
 
 	if (Materials.Num() > 0)
diff --git a/Engine/Source/Editor/VirtualTexturingEditor/Private/SConvertToVirtualTexture.cpp b/Engine/Source/Editor/VirtualTexturingEditor/Private/SConvertToVirtualTexture.cpp
new file mode 100644
--- /dev/null
+++ b/Engine/Source/Editor/VirtualTexturingEditor/Private/SConvertToVirtualTexture.cpp
@@ -0,0 +1,841 @@
+// Copyright Epic Games, Inc. All Rights Reserved.
+
+#include "SConvertToVirtualTexture.h"
+
+#include "ContentBrowserMenuContexts.h"
+#include "Engine/Texture2D.h"
+#include "Styling/AppStyle.h"
+#include "Materials/Material.h"
+#include "Materials/MaterialInstance.h"
+#include "Materials/MaterialFunctionInstance.h"
+#include "IContentBrowserSingleton.h"
+#include "ContentBrowserModule.h"
+#include "Particles/SubUVAnimation.h"
+
+#include "Editor.h"
+#include "Widgets/Text/STextBlock.h"
+#include "Widgets/Images/SImage.h"
+#include "Widgets/Layout/SBorder.h"
+#include "Misc/FeedbackContext.h"
+#include "Framework/Application/SlateApplication.h"
+#include "Widgets/Layout/SSeparator.h"
+#include "Widgets/Layout/SUniformGridPanel.h"
+#include "Widgets/Input/SButton.h"
+#include "Widgets/Layout/SScrollBox.h"
+#include "Widgets/Layout/SBox.h"
+#include "Widgets/Input/SComboBox.h"
+#include "Widgets/Input/SNumericEntryBox.h"
+#include "IAssetTools.h"
+#include "AssetRegistry/AssetRegistryHelpers.h"
+
+#define LOCTEXT_NAMESPACE "SConvertToVirtualTexture"
+
+void SConvertToVirtualTexture::Construct(const FArguments& InArgs)
+{
+	UserResponse = FConvertToVTDlg::Cancel;
+	ParentWindow = InArgs._ParentWindow.Get();
+	static FName ErrorIcon = "MessageLog.Error";
+
+
+	for (int i = 0; i < 16; i++)
+	{
+		TextureSizes.Add(MakeShareable(new int32(1 << i)));
+	}
+	ThresholdValue = *TextureSizes[10];
+
+	this->ChildSlot[
+		SNew(SVerticalBox)
+		// Textbox at the top giving an introductory message
+		+ SVerticalBox::Slot()
+			.AutoHeight()
+			.Padding(8.0f, 4.0f, 8.0f, 4.0f)
+			[
+				SNew(STextBlock)
+				.AutoWrapText(true)
+				.Visibility(this, &SConvertToVirtualTexture::GetIntroMessageVisibility)
+				.Text(this, &SConvertToVirtualTexture::GetIntroMessage)
+			]
+		// Error message at the top giving a common error message
+		+ SVerticalBox::Slot()
+			.AutoHeight()
+			.Padding(8.0f, 4.0f, 8.0f, 4.0f)
+			[
+				SNew(SHorizontalBox)
+				.Visibility(this, &SConvertToVirtualTexture::GetErrorMessageVisibility)
+				+ SHorizontalBox::Slot()
+				.AutoWidth()
+				.VAlign(VAlign_Center)
+				.HAlign(HAlign_Center)
+				[
+					SNew(SImage)
+					.Image(FAppStyle::GetBrush(ErrorIcon))
+				]
+				+ SHorizontalBox::Slot()
+				.AutoWidth()
+				.VAlign(VAlign_Center)
+				.HAlign(HAlign_Center)
+				[
+					SNew(STextBlock)
+					.AutoWrapText(true)
+					.Text(this, &SConvertToVirtualTexture::GetErrorMessage)
+				]
+			]
+		+ SVerticalBox::Slot()
+			.AutoHeight()
+			.Padding(8.0f, 4.0f, 8.0f, 4.0f)
+			[
+				SNew(SSeparator)
+			]
+		// The actual list of assets
+		+ SVerticalBox::Slot()
+			.Padding(8.0f, 4.0f, 8.0f, 4.0f)
+			[
+				SNew(SBorder)
+				[
+					SNew(SScrollBox)
+					+ SScrollBox::Slot()
+					[
+						SAssignNew(AssetListContainer, SVerticalBox)
+					]
+				]
+			]
+		// The bottom row of widgets: texture size selector
+		+ SVerticalBox::Slot()
+			.AutoHeight()
+			.HAlign(HAlign_Left)
+			.Padding(8.0f, 4.0f, 8.0f, 4.0f)
+			[
+				SNew(SHorizontalBox)
+				+SHorizontalBox::Slot()
+				.AutoWidth()
+				.VAlign(VAlign_Center)
+				[
+					SNew(STextBlock)
+					.Text(LOCTEXT("ConvertToVT_Size", "Texture size threshold: "))
+					.AutoWrapText(true)
+				]
+				+ SHorizontalBox::Slot()
+				.AutoWidth()
+				.VAlign(VAlign_Center)
+				[
+					SNew(SComboBox<TSharedPtr<int32>>)
+					.OptionsSource(&TextureSizes)
+					.OnSelectionChanged(this, &SConvertToVirtualTexture::OnThresholdChanged)
+					.OnGenerateWidget(this, &SConvertToVirtualTexture::OnGenerateThresholdWidget)
+					.InitiallySelectedItem(TextureSizes[10])
+					[
+						SNew(STextBlock)
+						.Text(this, &SConvertToVirtualTexture::GetThresholdText)
+					]
+				]
+				+ SHorizontalBox::Slot()
+				.AutoWidth()
+				.VAlign(VAlign_Center)
+				[
+					SNew(SButton)
+					.HAlign(HAlign_Center)
+					.ContentPadding(FAppStyle::GetMargin("StandardDialog.ContentPadding"))
+					.OnClicked(this, &SConvertToVirtualTexture::OnFilterButtonClicked)
+					.IsEnabled(this, &SConvertToVirtualTexture::GetFilterButtonEnabled)
+					.Text(LOCTEXT("ConvertToVT_Filter", "Apply Filter"))
+				]
+			]
+		// Separator
+		+ SVerticalBox::Slot()
+			.AutoHeight()
+			.Padding(8.0f, 4.0f, 8.0f, 4.0f)
+			[
+				SNew(SSeparator)
+			]
+		// Dialog ok/cancel buttons
+		+ SVerticalBox::Slot()
+			.AutoHeight()
+			.HAlign(HAlign_Right)
+			.Padding(8.0f, 4.0f, 8.0f, 4.0f)
+			[
+				SNew(SUniformGridPanel)
+				.SlotPadding(FAppStyle::GetMargin("StandardDialog.SlotPadding"))
+				.MinDesiredSlotWidth(FAppStyle::GetFloat("StandardDialog.MinDesiredSlotWidth"))
+				.MinDesiredSlotHeight(FAppStyle::GetFloat("StandardDialog.MinDesiredSlotHeight"))
+				+ SUniformGridPanel::Slot(0, 0)
+				[
+					SNew(SButton)
+					.HAlign(HAlign_Center)
+					.ContentPadding(FAppStyle::GetMargin("StandardDialog.ContentPadding"))
+					.OnClicked(this, &SConvertToVirtualTexture::OnButtonClick, FConvertToVTDlg::Confirm)
+					.IsEnabled(this, &SConvertToVirtualTexture::GetOkButtonEnabled)
+					.Text(LOCTEXT("ConvertToVT_OK", "OK"))
+				]
+				+ SUniformGridPanel::Slot(1, 0)
+				[
+					SNew(SButton)
+					.HAlign(HAlign_Center)
+					.ContentPadding(FAppStyle::GetMargin("StandardDialog.ContentPadding"))
+					.OnClicked(this, &SConvertToVirtualTexture::OnButtonClick, FConvertToVTDlg::Cancel)
+					.Text(LOCTEXT("ConvertToVT_Cancel", "Cancel"))
+				]
+			]
+	];
+
+	UpdateList();
+	bFilterButtonEnabled = false;
+}
+
+void SConvertToVirtualTexture::SetBackwards(bool bSetBackwards)
+{
+	Worker.SetConversionDirection(bSetBackwards);
+	bBackwards = bSetBackwards;
+}
+
+void SConvertToVirtualTexture::SetUserTextures(const TArray<UTexture2D *> &Textures)
+{
+	Worker.UserTextures = Textures;
+	UpdateList();
+}
+
+FConvertToVTDlg::EResult SConvertToVirtualTexture::GetUserResponse() const
+{
+	return UserResponse;
+}
+
+TSharedRef<SWidget> SConvertToVirtualTexture::CreateAssetLine(int index, const FAssetData &Asset, const FConversionStatus &Status)
+{
+	const bool bEngineAsset = Asset.PackagePath.ToString().StartsWith(TEXT("/Engine/"));
+	FName SeverityIcon = NAME_None;
+	FText DetailedInfoText;
+
+	if (Status.UserSelected)
+	{
+		if (Status.InvalidMaterialUsage)
+		{
+			SeverityIcon = "MessageLog.Error";
+			DetailedInfoText = LOCTEXT("ConvertToVT_ToolTip_InvalidUsage", "The texture could not be converted to VT due to its usage in materials (may be connected to a property that doesn't support VT).");
+		}
+		else if (Status.NonPowerOf2)
+		{
+			SeverityIcon = "MessageLog.Error";
+			DetailedInfoText = LOCTEXT("ConvertToVT_ToolTip_NonPowerOf2", "The texture could not be converted to VT because its size is not a power of 2.");
+		}
+		else if (bEngineAsset)
+		{
+			SeverityIcon = "MessageLog.Note";
+			DetailedInfoText = LOCTEXT("ConvertToVT_ToolTip_Engine", "The texture is an engine asset, a converted copy will be created in the current project.");
+		}
+		else if (Status.UnderSized)
+		{
+			SeverityIcon = "MessageLog.Note";
+			DetailedInfoText = LOCTEXT("ConvertToVT_ToolTip_Undersize", "The texture was under the threshold size but should still be converted to vt because it's using a shared parameter with other VT textures.");
+		}
+	}
+	else
+	{
+		if (Status.InvalidMaterialUsage)
+		{
+			SeverityIcon = "MessageLog.Warning";
+			DetailedInfoText = LOCTEXT("ConvertToVT_ToolTip_InvalidUsageNotSelected", "The texture could not be converted to VT due to its usage in materials (may be connected to a property that doesn't support VT).");
+		}
+		else if (Status.NonPowerOf2)
+		{
+			SeverityIcon = "MessageLog.Warning";
+			DetailedInfoText = LOCTEXT("ConvertToVT_ToolTip_NonPowerOf2NotSelected", "The texture could not be converted to VT because its size is not a power of 2.");
+		}
+		else if (bEngineAsset)
+		{
+			SeverityIcon = "MessageLog.Note";
+			DetailedInfoText = LOCTEXT("ConvertToVT_ToolTip_EngineNotSelected", "The engine asset texture was not selected but a converted copy will be created in the current project, because it's using a shared parameter with other VT textures.");
+		}
+		else if (Status.UnderSized)
+		{
+			SeverityIcon = "MessageLog.Note";
+			DetailedInfoText = LOCTEXT("ConvertToVT_ToolTip_UndersizeNotSelected", "The texture was not selected and furthermore is under the threshold size but should still be converted to vt because it's using a shared parameter with other VT textures.");
+		}
+		else
+		{
+			SeverityIcon = "MessageLog.Note";
+			DetailedInfoText = LOCTEXT("ConvertToVT_ToolTip_NotSelected", "The texture was not selected but should still be converted to vt because it's using a shared parameter with other VT textures.");
+		}
+	}
+
+	auto Result =
+		SNew(SVerticalBox)
+		+ SVerticalBox::Slot()
+		.AutoHeight()
+		.VAlign(VAlign_Center)
+		.HAlign(HAlign_Left)
+		[
+			// Class icon
+			SNew(SHorizontalBox)
+			+ SHorizontalBox::Slot()
+			.AutoWidth()
+			.VAlign(VAlign_Center)
+			.HAlign(HAlign_Center)
+			[
+				SNew(SBox)
+				.Padding(2.0f)
+				[
+					SNew(SImage)
+					.Image(FSlateIcon(FAppStyle::GetAppStyleSetName(),
+				(Asset.GetClass() == UTexture2D::StaticClass()) ?
+						"ClassIcon.Texture2D" :
+						((Asset.GetClass() == UMaterialFunction::StaticClass()) ? "ClassIcon.MaterialFunction" : "ClassIcon.Material")).GetIcon())
+				]
+			]
+			// Error/warning icon
+			+ SHorizontalBox::Slot()
+			.AutoWidth()
+			.HAlign(HAlign_Center)
+			.VAlign(VAlign_Center)
+			[
+				SNew(SBox)
+				.Padding(2.0f)
+				[
+					(SeverityIcon == NAME_None) ? SNullWidget::NullWidget :
+					static_cast<TSharedRef<SWidget>>(SNew(SImage).Image(FAppStyle::GetBrush(SeverityIcon)))
+				]
+			]
+			// Fold out button with asset name
+			+ SHorizontalBox::Slot()
+			.AutoWidth()
+			.HAlign(HAlign_Center)
+			.VAlign(VAlign_Center)
+			[
+					(DetailedInfoText.IsEmpty())
+					? static_cast<TSharedRef<SWidget>>(SNew(STextBlock)
+						.Text(FText::FromString(Asset.GetObjectPathString())))
+				: SNew(SButton)
+				.ButtonStyle(FCoreStyle::Get(), "NoBorder")
+				.VAlign(VAlign_Center)
+				.HAlign(HAlign_Center)
+				.ClickMethod(EButtonClickMethod::MouseDown)
+				.OnClicked(this, &SConvertToVirtualTexture::OnExpanderClicked, index)
+				.ContentPadding(0.f)
+				.ForegroundColor(FSlateColor::UseForeground())
+				.IsFocusable(false)
+				//.Text(FText::FromName(Asset.ObjectPath))
+				[
+					// Fold out icon
+					SNew(SHorizontalBox)
+					+ SHorizontalBox::Slot()
+					.AutoWidth()
+					.VAlign(VAlign_Center)
+					.HAlign(HAlign_Center)
+					[
+						SNew(SImage)
+						.Image(this, &SConvertToVirtualTexture::GetExpanderImage, index)
+					.ColorAndOpacity(FSlateColor::UseForeground())
+					]
+					// Fold out text
+					+ SHorizontalBox::Slot()
+					.AutoWidth()
+					.VAlign(VAlign_Center)
+					[
+						SNew(STextBlock)
+						.Text(FText::FromString(Asset.GetObjectPathString()))
+					]
+				]
+			]
+		]
+		+ SVerticalBox::Slot()
+		.AutoHeight()
+		.VAlign(VAlign_Center)
+		.HAlign(HAlign_Left)
+		[
+			SNew(STextBlock)
+			.Text(DetailedInfoText)
+			.Visibility(this, &SConvertToVirtualTexture::GetDetailVisibility, index)
+		]
+		+SVerticalBox::Slot()
+		.AutoHeight()
+		.VAlign(VAlign_Center)
+		.HAlign(HAlign_Left)
+		[
+			//SNew(STextBlock)
+			//.Text(GetAuditTrailText(Asset))
+			//.Visibility(this, &SConvertToVirtualTexture::GetDetailVisibility, index)
+			GetAuditTrailText(Asset, index)
+		];
+	return Result;
+}
+
+TSharedRef<SWidget> SConvertToVirtualTexture::GetAuditTrailText(const FAssetData &Asset, int32 index)
+{
+	//FString Result;
+	UObject *MaybeOk = Asset.GetAsset();
+	auto Trail = Worker.AuditTrail.Find(MaybeOk);
+
+	TSharedRef<SVerticalBox> Box = SNew(SVerticalBox)
+	.Visibility(this, &SConvertToVirtualTexture::GetDetailVisibility, index);
+
+	Box->AddSlot()
+	.AutoHeight()
+	.Padding(8.0f, 4.0f, 8.0f, 4.0f)
+	[
+		SNew(SSeparator)
+	];
+
+	Box->AddSlot()
+	.AutoHeight()
+	.VAlign(VAlign_Center)
+	.HAlign(HAlign_Left)
+	[
+		SNew(STextBlock)
+		.Text(LOCTEXT("ConvertToVT_AuditTrail","This texture was included because of the following dependencies:"))
+	];
+
+
+
+	while (Trail)
+	{
+		Box->AddSlot()
+		.AutoHeight()
+		.VAlign(VAlign_Center)
+		.HAlign(HAlign_Left)
+		[
+			SNew(STextBlock)
+			.Text(FText::FromString(Trail->PathDescription))
+		];
+
+		if (Trail->Destination)
+		{
+			Box->AddSlot()
+			.AutoHeight()
+			.VAlign(VAlign_Center)
+			.HAlign(HAlign_Left)
+			[
+				SNew(STextBlock)
+				.Text(FText::FromString(Trail->Destination->GetPathName()))
+			];
+
+			// May return null at which point we'll break out of the loop
+			auto NewTrail = Worker.AuditTrail.Find(Trail->Destination);
+			if (NewTrail == nullptr || NewTrail == Trail)
+			{
+				break;
+			}
+			Trail = NewTrail;
+		}
+		else
+		{
+			break;
+		}
+
+		//TODO: We just display the first step only as there isn't much more usefull info
+		//in the trail anyway for now and it could actually confuse people seeing half a trail.
+		break;
+	}
+
+	return Box;
+}
+
+void SConvertToVirtualTexture::UpdateList()
+{
+	Worker.FilterList(ThresholdValue);
+
+	AssetList.Empty();
+	AssetStatus.Empty();
+
+	for (UTexture2D *Texture : Worker.Textures)
+	{
+		AssetList.Add(Texture);
+		FConversionStatus* Status = new(AssetStatus) FConversionStatus();
+		Status->UserSelected = Worker.UserTextures.Contains(Texture);
+		Status->UnderSized = (Texture->GetSizeX()*Texture->GetSizeY() < ThresholdValue*ThresholdValue);
+		Status->NonPowerOf2 = !Texture->Source.IsPowerOfTwo() && Texture->PowerOfTwoMode == ETexturePowerOfTwoSetting::None;
+	}
+
+	for (UTexture2D *Texture : Worker.MaterialRejectedTextures)
+	{
+		AssetList.Add(Texture);
+		FConversionStatus* Status = new(AssetStatus) FConversionStatus();
+		Status->UserSelected = Worker.UserTextures.Contains(Texture);
+		Status->UnderSized = (Texture->GetSizeX()*Texture->GetSizeY() < ThresholdValue*ThresholdValue);
+		Status->NonPowerOf2 = !Texture->Source.IsPowerOfTwo() && Texture->PowerOfTwoMode == ETexturePowerOfTwoSetting::None;
+		Status->InvalidMaterialUsage = true;
+	}
+
+	for (UMaterial *RootMat : Worker.Materials)
+	{
+		// Don't want to show anything from the transient package, this may include preview materials from any currently active material editor
+		// We patch these up so material editor remains valid, but not useful to display this to user
+		if (RootMat->GetOutermost() != GetTransientPackage())
+		{
+			AssetList.Add(RootMat);
+			AssetStatus.Add(FConversionStatus());
+		}
+	}
+
+	for (UMaterialFunctionInterface *Func : Worker.Functions)
+	{
+		if (Func->GetOutermost() != GetTransientPackage())
+		{
+			AssetList.Add(Func);
+			AssetStatus.Add(FConversionStatus());
+		}
+	}
+
+	check(AssetList.Num() == AssetStatus.Num());
+	AssetListContainer->ClearChildren();
+
+	for (int Id = 0; Id < AssetList.Num(); Id++)
+	{
+		AssetListContainer->AddSlot()
+			.AutoHeight()
+			[
+				CreateAssetLine(Id, AssetList[Id], AssetStatus[Id])
+			];
+	}
+
+	ErrorMessage = FText();
+}
+
+
+FReply SConvertToVirtualTexture::OnButtonClick(FConvertToVTDlg::EResult ButtonID)
+{
+	ParentWindow->RequestDestroyWindow();
+	UserResponse = ButtonID;
+
+	if (ButtonID == FConvertToVTDlg::EResult::Confirm)
+	{
+		Worker.DoConvert();
+	}
+
+	return FReply::Handled();
+}
+
+void SConvertToVirtualTexture::OnThresholdChanged(TSharedPtr<int32> InSelectedItem, ESelectInfo::Type SelectInfo)
+{
+	ThresholdValue = *InSelectedItem;
+	bFilterButtonEnabled = true;
+}
+
+FText SConvertToVirtualTexture::GetThresholdText() const
+{
+	return FText::FromString(FString::Format(TEXT("{0}"), TArray<FStringFormatArg>({ ThresholdValue })));
+}
+
+TSharedRef<SWidget> SConvertToVirtualTexture::OnGenerateThresholdWidget(TSharedPtr<int32> InItem)
+{
+	return SNew(STextBlock)
+		.Text(FText::FromString(FString::Format(TEXT("{0}"), TArray<FStringFormatArg>({ *InItem }))));
+}
+
+bool SConvertToVirtualTexture::GetFilterButtonEnabled() const
+{
+	return bFilterButtonEnabled;
+}
+
+FReply SConvertToVirtualTexture::OnFilterButtonClicked()
+{
+	UpdateList();
+	bFilterButtonEnabled = false;
+
+	return FReply::Handled();
+}
+
+FReply SConvertToVirtualTexture::OnExpanderClicked(int index)
+{
+	if (ExpandedIndexes.Contains(index))
+	{
+		ExpandedIndexes.Remove(index);
+	}
+	else
+	{
+		ExpandedIndexes.Add(index);
+	}
+	return FReply::Handled();
+}
+
+bool SConvertToVirtualTexture::GetOkButtonEnabled() const
+{
+	return ErrorMessage.IsEmpty();
+}
+
+EVisibility SConvertToVirtualTexture::GetDetailVisibility(int index) const
+{
+	return (ExpandedIndexes.Contains(index)) ? EVisibility::Visible : EVisibility::Collapsed;
+}
+
+const FSlateBrush* SConvertToVirtualTexture::GetExpanderImage(int index) const
+{
+	FName ResourceName;
+	if (GetDetailVisibility(index) == EVisibility::Visible)
+	{
+		static FName ExpandedName = "TreeArrow_Expanded";
+		ResourceName = ExpandedName;
+	}
+	else
+	{
+		static FName CollapsedName = "TreeArrow_Collapsed";
+		ResourceName = CollapsedName;
+	}
+	return FCoreStyle::Get().GetBrush(ResourceName);
+}
+
+EVisibility SConvertToVirtualTexture::GetIntroMessageVisibility() const
+{
+	return (IntroMessage.IsEmpty()) ? EVisibility::Collapsed : EVisibility::Visible;
+}
+
+EVisibility SConvertToVirtualTexture::GetErrorMessageVisibility() const
+{
+	return (ErrorMessage.IsEmpty()) ? EVisibility::Collapsed : EVisibility::Visible;
+}
+
+EVisibility SConvertToVirtualTexture::GetThresholdVisibility() const
+{
+	return (bThresholdVisible) ? EVisibility::Collapsed : EVisibility::Visible;
+}
+
+FText SConvertToVirtualTexture::GetIntroMessage() const
+{
+	return IntroMessage;
+}
+
+FText SConvertToVirtualTexture::GetErrorMessage() const
+{
+	return ErrorMessage;
+}
+
+FConvertToVTDlg::FConvertToVTDlg(const TArray<UTexture2D *> &Textures, bool bBackwards)
+{
+	if (FSlateApplication::IsInitialized())
+	{
+		DialogWindow = SNew(SWindow)
+			.Title((bBackwards) ? LOCTEXT("ConvertToVTDlgTitle_Backwards", "Convert From VT") : LOCTEXT("ConvertToVTDlgTitle", "Convert To VT"))
+			.SupportsMinimize(false).SupportsMaximize(false)
+			.ClientSize(FVector2D(500, 500));
+
+		TSharedPtr<SBorder> DialogWrapper =
+			SNew(SBorder)
+			.BorderImage(FAppStyle::GetBrush("ToolPanel.GroupBorder"))
+			.Padding(4.0f)
+			[
+				SAssignNew(DialogWidget, SConvertToVirtualTexture)
+				.ParentWindow(DialogWindow)
+			];
+
+		DialogWidget->SetBackwards(bBackwards);
+		DialogWidget->SetUserTextures(Textures);
+		DialogWindow->SetContent(DialogWrapper.ToSharedRef());
+	}
+}
+
+FConvertToVTDlg::EResult FConvertToVTDlg::ShowModal()
+{
+	//Show Dialog
+	GEditor->EditorAddModalWindow(DialogWindow.ToSharedRef());
+	EResult UserResponse = (EResult)DialogWidget->GetUserResponse();
+	DialogWindow->GetParentWindow()->RemoveDescendantWindow(DialogWindow.ToSharedRef());
+	return UserResponse;
+}
+
+
+void SConvertToVirtualTexture::ConvertVTTexture(TArray<UTexture2D*> Objects, bool backwards)
+{
+	TArray<FAssetData> AllRelevantMaterials;
+	TArray<UTexture2D*> UserTextures; // The original selection of the user
+
+	//FScopedSlowTask SlowTask(4.0f, LOCTEXT("ConvertToVT_Progress_ConvertToVt", "Converting textures to VT..."));
+	//SlowTask.MakeDialog();
+
+	// First simply assemble a list of materials to show the user.
+	//SlowTask.EnterProgressFrame();
+	for (auto ObjIt = Objects.CreateConstIterator(); ObjIt; ++ObjIt)
+	{
+		auto Object = (*ObjIt);
+		UTexture2D *Texture = Cast<UTexture2D>(Object);
+		if (Texture != nullptr && Texture->VirtualTextureStreaming == backwards)
+		{
+			UAssetRegistryHelpers::FindReferencersOfAssetOfClass(Object, { UMaterialInterface::StaticClass() }, AllRelevantMaterials);
+			UserTextures.Add(Texture);
+		}
+	}
+
+	FConvertToVTDlg ConvertDlg(UserTextures, backwards);
+	/*ConvertDlg.AssetList.Append(UserTextures);
+	ConvertDlg.AssetList.Append(AllRelevantMaterials);
+	ConvertDlg.AssetStatus.AddDefaulted(UserTextures.Num());
+	ConvertDlg.AssetStatus.AddDefaulted(AllRelevantMaterials.Num());
+	ConvertDlg.Message = LOCTEXT("ConvertToVTMaterialList", "The following textures & materials (and their dependencies) will be loaded by this operation. This could take a while.");
+	ConvertDlg.Threshold = virtualTextureAutoEnableThreshold;*/
+
+	if (ConvertDlg.ShowModal() == FConvertToVTDlg::Confirm)
+	{
+#if 0
+		//FVTConversionWorker Conversion;
+		//Conversion.UserTextures = UserTextures;
+		//Conversion.FilterList(ConvertDlg.Threshold);
+
+		/*FScopedSlowTask SlowTask(3.0f, LOCTEXT("ConvertToVT_Progress_ConvertToVt", "Converting textures to VT..."));
+		SlowTask.MakeDialog();
+
+		TArray<UTexture2D*> Textures; // All textures that should be converted
+		TArray<UTexture2D*> SizeRejectedTextures; // The original selection of the user filtered by textures that match the threshold size
+
+		for (UTexture2D *Texture: UserTextures)
+		{
+			if (Texture->GetSizeX()*Texture->GetSizeY() >= ConvertDlg.Threshold*ConvertDlg.Threshold)
+			{
+				Textures.Add(Texture);
+			}
+			else
+			{
+				SizeRejectedTextures.Add(Texture);
+			}
+		}
+
+		SlowTask.EnterProgressFrame();
+		TArray<UMaterial *>Materials;
+		TArray<UMaterialFunctionInterface *> Functions;
+//		FindAllTexturesAndMaterials(Materials, Functions, Textures);
+*/
+		// Show the dialog yet again warning the user...
+
+		/*FConvertToVTDlg ConvertDlg2;
+		ConvertDlg2.AssetList.Empty();
+		for (UTexture2D *Texture : Conversion.Textures)
+		{
+			ConvertDlg2.AssetList.Add(Texture);
+			ConvertDlg2.AssetStatus.Add(FConversionStatus(UserTextures.Contains(Texture), (Texture->GetSizeX()*Texture->GetSizeY() < ConvertDlg.Threshold*ConvertDlg.Threshold)));
+		}
+
+		for (UMaterial *RootMat : Conversion.Materials)
+		{
+			ConvertDlg2.AssetList.Add(RootMat);
+			ConvertDlg2.AssetStatus.Add(FConversionStatus());
+		}
+
+		for (UMaterialFunctionInterface *Func : Conversion.Functions)
+		{
+			ConvertDlg2.AssetList.Add(Func);
+			ConvertDlg2.AssetStatus.Add(FConversionStatus());
+		}
+
+		ConvertDlg2.Message = LOCTEXT("ConvertToVTTextureList", "The following textures and materials will converted. Please carefully inspect this list as textures not originally selected may also be included for technical reasons.");
+		ConvertDlg2.Threshold = ConvertDlg.Threshold;
+		ConvertDlg2.EditThreshold = false;
+		if (ConvertDlg2.ShowModal() == FConvertToVTDlg::Confirm)
+		{
+			Conversion.DoConvert(backwards);
+			/*UE_LOG(LogVirtualTextureConversion, Display, TEXT("Beginning conversion..."));
+			SlowTask.EnterProgressFrame();
+			{
+				FScopedSlowTask TextureTask(Textures.Num(), LOCTEXT("ConvertToVT_Progress_TextureTask", "Updating textures..."));
+
+				for (UTexture2D *Tex : Textures)
+				{
+					UE_LOG(LogVirtualTextureConversion, Display, TEXT("Texture %s"), *Tex->GetName());
+					TextureTask.EnterProgressFrame();
+
+					bool OldVt = Tex->VirtualTextureStreaming;
+					Tex->VirtualTextureStreaming = !backwards;
+					if (OldVt != Tex->VirtualTextureStreaming)
+					{
+						Tex->Modify();
+						Tex->PostEditChange();
+					}
+				}
+			}
+
+			SlowTask.EnterProgressFrame();
+			{
+				FScopedSlowTask MaterialTask(Materials.Num() + Functions.Num(), LOCTEXT("ConvertToVT_Progress_MaterialTask", "Updating materials..."));
+				FMaterialUpdateContext UpdateContext;
+
+				TMap<UMaterialFunctionInterface*, TArray<UMaterial*>> FunctionToMaterialMap;
+
+				for (UMaterial *Mat : Materials)
+				{
+					UE_LOG(LogVirtualTextureConversion, Display, TEXT("Material %s"), *Mat->GetName());
+
+					MaterialTask.EnterProgressFrame();
+
+					bool MatModified = false;
+					for (UMaterialExpression *Expr : Mat->Expressions)
+					{
+						UMaterialExpressionTextureBase *TexExpr = Cast<UMaterialExpressionTextureBase>(Expr);
+						if (TexExpr)
+						{
+							if (Textures.Contains(TexExpr->Texture))
+							{
+								UE_LOG(LogVirtualTextureConversion, Display, TEXT("Adjusting sampler %s."), *TexExpr->GetName());
+								auto OldType = TexExpr->SamplerType;
+								TexExpr->AutoSetSampleType();
+								if (TexExpr->SamplerType != OldType)
+								{
+									TexExpr->Modify();
+									TexExpr->PostEditChange();
+									MatModified = true;
+								}
+							}
+						}
+					}
+
+					TArray<UMaterialFunctionInterface *>MaterialFunctions;
+					Mat->GetDependentFunctions(MaterialFunctions);
+					for (auto Function : MaterialFunctions)
+					{
+						FunctionToMaterialMap.FindOrAdd(Function).AddUnique(Mat);
+					}
+
+					if (MatModified)
+					{
+						UE_LOG(LogVirtualTextureConversion, Display, TEXT("Material %s added to update list."), *Mat->GetName());
+						UMaterialEditingLibrary::RecompileMaterial(Mat);
+					}
+					else
+					{
+						UE_LOG(LogVirtualTextureConversion, Display, TEXT("Material %s was not modified, skipping."), *Mat->GetName());
+					}
+				}
+
+				for (UMaterialFunctionInterface *Func : Functions)
+				{
+					UE_LOG(LogVirtualTextureConversion, Display, TEXT("Function %s"), *Func->GetName());
+
+					MaterialTask.EnterProgressFrame();
+
+					bool FuncModified = false;
+					const TArray<UMaterialExpression*> *Expressions = Func->GetExpressions();
+					for (UMaterialExpression *Expr : *Expressions)
+					{
+						UMaterialExpressionTextureBase *TexExpr = Cast<UMaterialExpressionTextureBase>(Expr);
+						if (TexExpr)
+						{
+							if (Textures.Contains(TexExpr->Texture))
+							{
+								UE_LOG(LogVirtualTextureConversion, Display, TEXT("Adjusting sampler %s."), *TexExpr->GetName());
+								auto OldType = TexExpr->SamplerType;
+								TexExpr->AutoSetSampleType();
+								if (TexExpr->SamplerType != OldType)
+								{
+									TexExpr->Modify();
+									TexExpr->PostEditChange();
+									FuncModified = true;
+								}
+							}
+						}
+					}
+
+					if (FuncModified)
+					{
+						UMaterialEditingLibrary::UpdateMaterialFunction(Func, nullptr);
+						UE_LOG(LogVirtualTextureConversion, Display, TEXT("Function %s added to update list."), *Func->GetName());
+					}
+					else
+					{
+						UE_LOG(LogVirtualTextureConversion, Display, TEXT("Function %s was not modified, skipping."), *Func->GetName());
+					}
+				}
+
+				// update the world's viewports
+				UE_LOG(LogVirtualTextureConversion, Display, TEXT("Broadcasting to editor."));
+				FEditorDelegates::RefreshEditor.Broadcast();
+				FEditorSupportDelegates::RedrawAllViewports.Broadcast();
+			}
+		}*/
+		#endif
+	}
+}
+
+#undef LOCTEXT_NAMESPACE
\ No newline at end of file
diff --git a/Engine/Source/Editor/VirtualTexturingEditor/Private/SConvertToVirtualTexture.h b/Engine/Source/Editor/VirtualTexturingEditor/Private/SConvertToVirtualTexture.h
new file mode 100644
--- /dev/null
+++ b/Engine/Source/Editor/VirtualTexturingEditor/Private/SConvertToVirtualTexture.h
@@ -0,0 +1,149 @@
+// Copyright Epic Games, Inc. All Rights Reserved.
+
+#pragma once
+
+#include "CoreMinimal.h"
+#include "SlateFwd.h"
+#include "UObject/Object.h"
+#include "Layout/Visibility.h"
+#include "Widgets/DeclarativeSyntaxSupport.h"
+#include "Input/Reply.h"
+#include "Styling/AppStyle.h"
+#include "Layout/Visibility.h"
+#include "Widgets/SWidget.h"
+#include "Widgets/SCompoundWidget.h"
+#include "Widgets/Views/STableViewBase.h"
+#include "Widgets/Views/STableRow.h"
+#include "Widgets/Views/STreeView.h"
+#include "GameplayTagsManager.h"
+#include "VirtualTextureConversionWorker.h"
+
+class UTexture2D;
+
+struct FConversionStatus
+{
+	bool UserSelected = true; // Originally selected by the user
+	bool UnderSized = false; // Too small to convert according to user settings
+	bool NonPowerOf2 = false; // Not a power of 2, can't convert
+	bool InvalidMaterialUsage = false;
+};
+
+/**
+* FConvertToVTDlg
+*
+* Wrapper class for SConvertToVTDlg. This class creates and launches a dialog then awaits the
+* result to return to the user.
+*/
+class FConvertToVTDlg
+{
+public:
+	enum EResult
+	{
+		Cancel = 0,			// No/Cancel, normal usage would stop the current action
+		Confirm = 1,		// Yes/Ok/Etc, normal usage would continue with action
+	};
+
+	FConvertToVTDlg(const TArray<UTexture2D*>& Textures, bool bBackwards);
+
+	/**  Shows the dialog box and waits for the user to respond. */
+	EResult ShowModal();
+
+private:
+	TSharedPtr<SWindow> DialogWindow;
+	TSharedPtr<class SConvertToVirtualTexture> DialogWidget;
+};
+
+struct FMaterialItem
+{
+	FName Name;
+};
+
+class SConvertToVirtualTexture : public SCompoundWidget
+{
+public:
+	static void ConvertVTTexture(TArray<UTexture2D*> Objects, bool backwards);
+	
+public:
+
+	SLATE_BEGIN_ARGS(SConvertToVirtualTexture)
+	{}
+	/** Window in which this widget resides */
+	SLATE_ATTRIBUTE(TSharedPtr<SWindow>, ParentWindow)
+	SLATE_END_ARGS()
+
+	void Construct(const FArguments& InArgs);
+
+	void SetBackwards(bool bSetBackwards);
+
+	void SetUserTextures(const TArray<UTexture2D *> &Textures);
+
+	FConvertToVTDlg::EResult GetUserResponse() const;
+
+private:
+
+	/**
+	* Creates a single line showing an asset and it's status related to VT conversion
+	*/
+	TSharedRef<SWidget> CreateAssetLine(int index, const FAssetData &Asset, const FConversionStatus &Status);
+
+	TSharedRef<SWidget> GetAuditTrailText(const FAssetData &Asset, int32 index);
+
+	void UpdateList();
+
+	FReply OnButtonClick(FConvertToVTDlg::EResult ButtonID);
+
+	void OnThresholdChanged(TSharedPtr<int32> InSelectedItem, ESelectInfo::Type SelectInfo);
+
+	FText GetThresholdText() const;
+
+	TSharedRef<SWidget> OnGenerateThresholdWidget(TSharedPtr<int32> InItem);
+
+	bool GetFilterButtonEnabled() const;
+
+	FReply OnFilterButtonClicked();
+	
+	FReply OnExpanderClicked(int index);
+
+	bool GetOkButtonEnabled() const;
+	
+	EVisibility GetDetailVisibility(int index) const;
+
+	const FSlateBrush* GetExpanderImage(int index) const;
+
+	EVisibility GetIntroMessageVisibility() const;
+
+	EVisibility GetErrorMessageVisibility() const;
+
+	EVisibility GetThresholdVisibility() const;
+
+	FText GetIntroMessage() const;
+
+	FText GetErrorMessage() const;
+	
+	TArray<int> ExpandedIndexes;
+
+	bool bBackwards;
+	FText IntroMessage;
+	FText ErrorMessage;
+	FConvertToVTDlg::EResult	 UserResponse;
+
+	TSharedPtr<SVerticalBox>	 AssetListContainer;
+	//TSharedPtr<STextBlock>		 MessageTextBlock;
+	//TSharedPtr<SHorizontalBox>	 ThresholdContainer;
+	//TSharedPtr<SHorizontalBox>	 ErrorContainer;
+	//TSharedPtr<SButton>			 OkButton;
+
+	/** Pointer to the window which holds this Widget, required for modal control */
+	TSharedPtr<SWindow>			 ParentWindow;
+
+	FVirtualTextureConversionWorker Worker;
+	TArray<FAssetData> AssetList;
+	TArray<FConversionStatus> AssetStatus;
+
+	int ThresholdValue;
+	bool bThresholdVisible;
+
+	TArray<TSharedPtr<int32>> TextureSizes;
+
+	bool bFilterButtonEnabled;
+};
diff --git a/Engine/Source/Editor/VirtualTexturingEditor/Private/VirtualTextureConversionWorker.cpp b/Engine/Source/Editor/VirtualTexturingEditor/Private/VirtualTextureConversionWorker.cpp
new file mode 100644
--- /dev/null
+++ b/Engine/Source/Editor/VirtualTexturingEditor/Private/VirtualTextureConversionWorker.cpp
@@ -0,0 +1,696 @@
+// Copyright Epic Games, Inc. All Rights Reserved.
+
+#include "VirtualTextureConversionWorker.h"
+
+#include "Engine/Texture2D.h"
+#include "Misc/PackageName.h"
+#include "Materials/Material.h"
+#include "Materials/MaterialInstance.h"
+#include "Materials/MaterialFunctionInstance.h"
+#include "MaterialEditor/PreviewMaterial.h"
+#include "Factories/MaterialFactoryNew.h"
+#include "EditorFramework/AssetImportData.h"
+#include "AssetRegistry/AssetRegistryModule.h"
+#include "ObjectTools.h"
+
+#include "Editor.h"
+#include "Misc/ScopedSlowTask.h"
+#include "EditorSupportDelegates.h"
+#include "MaterialGraph/MaterialGraph.h"
+#include "MaterialEditingLibrary.h"
+#include "AssetRegistry/AssetRegistryHelpers.h"
+#include "UObject/UObjectIterator.h"
+
+#define LOCTEXT_NAMESPACE "AssetVTConversion"
+
+DEFINE_LOG_CATEGORY_STATIC(LogVirtualTextureConversion, Log, All);
+
+template <class T> void GetReferencersOfType(UObject *Object, TArray<T*> &OutObjects)
+{
+	TArray<FAssetData> OutAssetDatas;
+	UAssetRegistryHelpers::FindReferencersOfAssetOfClass(Object, { T::StaticClass() }, OutAssetDatas);
+
+	FScopedSlowTask SlowTask(OutAssetDatas.Num(), LOCTEXT("ConvertToVT_Progress_LoadingObjects", "Loading Objects..."));
+
+	for (auto Data : OutAssetDatas)
+	{
+		SlowTask.EnterProgressFrame();
+		UObject *MaybeOk = Data.GetAsset();
+
+		// It's a material?
+		T *IsOk = Cast<T>(MaybeOk);
+
+		if (IsOk != nullptr)
+		{
+			OutObjects.Add(IsOk);
+		}
+	}
+}
+
+static void GetPreviewMaterials(UTexture2D* InTexture, TArray<UMaterialInterface*>& OutMaterials)
+{
+	for (TObjectIterator<UPreviewMaterial> It; It; ++It)
+	{
+		UMaterial* Material = *It;
+		for (uint32 PropertyIndex = 0u; PropertyIndex < MP_MAX; ++PropertyIndex)
+		{
+			const EMaterialProperty PropertyToValidate = (EMaterialProperty)PropertyIndex;
+			if (Material->IsTextureReferencedByProperty(PropertyToValidate, InTexture))
+			{
+				OutMaterials.Add(Material);
+				break;
+			}
+		}
+	}
+}
+
+bool IsVirtualTextureValidForMaterial(UMaterialInterface* InMaterial, UTexture2D* InTexture)
+{
+	for (uint32 PropertyIndex = 0u; PropertyIndex < MP_MAX; ++PropertyIndex)
+	{
+		const EMaterialProperty PropertyToValidate = (EMaterialProperty)PropertyIndex;
+		const EShaderFrequency ShaderFrequencyToValidate = FMaterialAttributeDefinitionMap::GetShaderFrequency(PropertyToValidate);
+		if (PropertyToValidate == MP_OpacityMask || ShaderFrequencyToValidate != SF_Pixel)
+		{
+			// see if the texture is referenced by a property that doesn't support virtual texture access
+			if (InMaterial->IsTextureReferencedByProperty(PropertyToValidate, InTexture))
+			{
+				return false;
+			}
+		}
+	}
+
+	return true;
+}
+
+void FVirtualTextureConversionWorker::AddReferencedObjects( FReferenceCollector& Collector )
+{
+	Collector.AddReferencedObjects( UserTextures );
+	Collector.AddReferencedObjects( Textures );
+	Collector.AddReferencedObjects( Materials );
+	Collector.AddReferencedObjects( Functions );
+	Collector.AddReferencedObjects( SizeRejectedTextures );
+	Collector.AddReferencedObjects( SizeRejectedTextures );
+	Collector.AddReferencedObjects( MaterialRejectedTextures );
+
+	TArray<UObject*> AuditTrailKeys;
+	AuditTrail.GetKeys(AuditTrailKeys);
+	Collector.AddReferencedObjects( AuditTrailKeys );
+}
+
+void FVirtualTextureConversionWorker::FindAllTexturesAndMaterials_Iteration(TArray<UMaterial*>& InAffectedMaterials,
+	TArray<UMaterialFunctionInterface*>& InAffectedFunctions,
+	TArray<UTexture2D*>& InAffectedTextures,
+	TArray<UTexture2D*>& InInvalidTextures,
+	FScopedSlowTask& Task)
+{
+	TArray<UMaterialInterface *>MaterialInferfaces; // All parents and instances
+	TArray<UMaterialInterface *>MaterialHeap; // Temporary work heap
+	TArray<UMaterialFunctionInterface *>FunctionInferfaces; // All parents and instances
+	TArray<UMaterialFunctionInterface *>FunctionHeap; // Temporary work heap
+	TMap<UMaterial *, TArray<FMaterialParameterInfo>> ParametersToVtIze;
+	TMap<UMaterialFunctionInterface *, TArray<FMaterialParameterInfo>> FunctionParametersToVtIze;
+
+	// Find all materials that reference the passed in textures
+	// This will also load these materials.
+	TArray<UMaterialInterface *> MaterialsUsedByAffectedTextures;
+	int TextureIndex = 0;
+	while (TextureIndex < InAffectedTextures.Num())
+	{
+		UTexture2D* Tex2D = InAffectedTextures[TextureIndex];
+
+		Task.EnterProgressFrame();
+		TArray<UMaterialInterface *> MaterialsUsingTexture;
+		GetReferencersOfType(Tex2D, MaterialsUsingTexture);
+
+		/*for (auto Material : MaterialsUsingTexture)
+		{
+			AuditTrail.Add(Material, FAuditTrail(
+				Tex2D,
+				FString::Printf(TEXT("references texture"))
+			));
+		}*/
+
+		// Check all materials using texture, make sure we're able to convert this texture
+		bool bIsVirtualTextureValid = true;
+		for (UMaterialInterface* Material : MaterialsUsingTexture)
+		{
+			if (!bConvertBackward && !IsVirtualTextureValidForMaterial(Material, Tex2D))
+			{
+				AuditTrail.Add(Tex2D, FAuditTrail(
+					Material,
+					TEXT("does not support VT usage on material")));
+				bIsVirtualTextureValid = false;
+				break;
+			}
+		}
+
+		if (bIsVirtualTextureValid)
+		{
+			// If this is an engine texture, we'll make a copy and update any needed references to point to the new copy
+			// since we're not changing the original texture, we don't need to bring in any additional dependencies
+			// Non-power-2 textures won't convert to VT, don't bring any dependencies for them
+			if (!Tex2D->GetPathName().StartsWith("/Engine/") &&
+				(Tex2D->Source.IsPowerOfTwo() || Tex2D->PowerOfTwoMode != ETexturePowerOfTwoSetting::None))
+			{
+				// Also get any preview materials that reference the given texture
+				// We need to convert these to ensure any active material editors remain valid
+				GetPreviewMaterials(Tex2D, MaterialsUsedByAffectedTextures);
+
+				MaterialsUsedByAffectedTextures.Append(MaterialsUsingTexture);
+			}
+			++TextureIndex;
+		}
+		else
+		{
+			InAffectedTextures.RemoveAtSwap(TextureIndex);
+			InInvalidTextures.Add(Tex2D);
+		}
+	}
+
+	// Find all materials that reference the passed in functions
+	// This will also load these materials.
+	for (auto Func : InAffectedFunctions)
+	{
+		Task.EnterProgressFrame();
+		TArray<UMaterialInterface *> MaterialsUsingFunction;
+		GetReferencersOfType(Func, MaterialsUsingFunction);
+
+		/*for (auto Material : MaterialsUsingFunction)
+		{
+			AuditTrail.Add(Material, FAuditTrail(
+				Func,
+				FString::Printf(TEXT("references function"))
+			));
+		}*/
+
+		MaterialsUsedByAffectedTextures.Append(MaterialsUsingFunction);
+	}
+
+	// Find all the root materials of the found instances and add them to our
+	// working lists.
+	for (UMaterialInterface *If : MaterialsUsedByAffectedTextures)
+	{
+		// It's a material?
+		UMaterial *Material = Cast<UMaterial>(If);
+
+		// It's a material instance? we're only interested in it's root for now
+		UMaterialInstance *MaterialInstance = Cast<UMaterialInstance>(If);
+		if (MaterialInstance != nullptr)
+		{
+			Material = MaterialInstance->GetMaterial();
+		}
+
+		check(Material); // It's something else entirely??
+		MaterialInferfaces.AddUnique(Material);
+		MaterialHeap.AddUnique(Material);
+		InAffectedMaterials.AddUnique(Material);
+
+		AuditTrail.Add(Material, FAuditTrail(
+			If,
+			FString::Printf(TEXT("is the child  of"))
+		));
+	}
+
+	// We now have a set of "root" materials which will be affected by changing InTextures to VT.
+	// Now find all children of these materials which will also be affected trough parameters now requiring VT textures being set on them.
+	// This will again load any child instances and their dependencies which aren't loaded yet
+	TArray<UMaterialInterface *>VisitedMaterials;
+
+	while (MaterialHeap.Num() > 0)
+	{
+		UMaterialInterface *ParentMaterial = MaterialHeap[0];
+		MaterialHeap.RemoveAt(0);
+
+		if (VisitedMaterials.Contains(ParentMaterial))
+		{
+			UE_LOG(LogVirtualTextureConversion, Display, TEXT("Circular inheritance chain!? %s"), *ParentMaterial->GetName());
+			continue;
+		}
+
+		VisitedMaterials.Add(ParentMaterial);
+
+		// Check all parameters of the current material. If they reference a texture
+		// we want to convert to VT flag the parameter (this will then cause all textures assigned to this parameter to convert to vt as well)
+		TArray<FMaterialParameterInfo> ParameterInfos;
+		TArray<FGuid> ParameterGuids;
+		ParentMaterial->GetAllTextureParameterInfo(ParameterInfos, ParameterGuids);
+
+		for (auto ParamInfo : ParameterInfos)
+		{
+			UTexture *ParamValue = nullptr;
+			if (ParentMaterial->GetTextureParameterValue(ParamInfo, ParamValue))
+			{
+				UTexture2D *ParamValue2D = Cast<UTexture2D>(ParamValue);
+				if (ParamValue2D != nullptr)
+				{
+					if (InAffectedTextures.Contains(ParamValue2D))
+					{
+						//UE_LOG(LogVirtualTextureConversion, Display, TEXT("Adding parameter %s because it references  %s on %s >> %s"), *ParamInfo.Name.ToString(), *ParamValue2D->GetPathName(), *ParentMaterial->GetMaterial()->GetPathName(), *ParentMaterial->GetPathName());
+						ParametersToVtIze.FindOrAdd(ParentMaterial->GetMaterial()).Add(ParamInfo);
+					}
+				}
+			}
+		}
+
+		// Find all direct children of this material (children of children will be discovered later by pushing these children on the MaterialHeap).
+		TArray<UMaterialInstance*> ParentMaterialInstances;
+		Task.EnterProgressFrame();
+		GetReferencersOfType(ParentMaterial, ParentMaterialInstances);
+
+		for (auto MaterialInstance : ParentMaterialInstances)
+		{
+			//MaterialInstances.AddUnique(MaterialInstance);
+			MaterialInferfaces.AddUnique(MaterialInstance);
+			//InstancesForMaterial.FindOrAdd(MaterialInstance->GetMaterial()).AddUnique(MaterialInstance);
+
+			// Push on the heap to check materials referencing us recursively
+			MaterialHeap.AddUnique(MaterialInstance);
+		}
+	}
+	// We know have a set of root materials and a set of properties to convert to VT
+	// Find all textures referenced by these properties
+	// These new textures could in turn be referenced by other materials (not in the inheritance chain)
+	// which is why we have to run this discovery process iteratively until we don't discover any new
+	// materials or textures.
+
+	for (UMaterialInterface *If : MaterialInferfaces)
+	{
+		UMaterial *Mat = If->GetMaterial();
+		for (const FMaterialParameterInfo &Parameter : ParametersToVtIze.FindOrAdd(Mat))
+		{
+			UTexture *Tex = nullptr;
+			if (If->GetTextureParameterValue(Parameter, Tex))
+			{
+				UTexture2D *Tex2D = Cast<UTexture2D>(Tex);
+				if (Tex2D && !Tex2D->VirtualTextureStreaming)
+				{
+					InAffectedTextures.AddUnique(Tex2D);
+					AuditTrail.Add(Tex2D, FAuditTrail(
+						Mat,
+						FString::Printf(TEXT("set on parameter %s in instance %s of material"), *Parameter.Name.ToString(), *If->GetName())
+					));
+				}
+			}
+		}
+	}
+
+	//
+	// Pretty much the same again but now for material functions....
+	//
+
+	// Find all materials functions that directly reference the passed in textures
+	TArray<UMaterialFunctionInterface *> FunctionsUsedByAffectedTextures;
+	for (auto Tex2D : InAffectedTextures)
+	{
+		Task.EnterProgressFrame();
+		if (!Tex2D->GetPathName().StartsWith("/Engine/"))
+		{
+			GetReferencersOfType(Tex2D, FunctionsUsedByAffectedTextures);
+		}
+	}
+
+	// Find all the root function of the found instances and add them to our
+	// working lists.
+	for (UMaterialFunctionInterface *If : FunctionsUsedByAffectedTextures)
+	{
+		// It's a material?
+		UMaterialFunction* Function = If->GetBaseFunction();
+		if (Function)
+		{
+			FunctionInferfaces.AddUnique(Function);
+			FunctionHeap.AddUnique(Function);
+			InAffectedFunctions.AddUnique(Function);
+
+			AuditTrail.Add(Function, FAuditTrail(
+				If,
+				FString::Printf(TEXT("this is the parent of"))
+			));
+		}
+	}
+
+
+	// We have a second class of functions here. That is functions that don't directly reference any textures of interest
+	// but where there are material instances that overrides properties in the textures that do reference textures of interest
+	for (auto MaterialParametersPair : ParametersToVtIze)
+	{
+		for (auto Parameter : MaterialParametersPair.Value)
+		{
+			TArray<UMaterialFunctionInterface*> DependentFunctions;
+			MaterialParametersPair.Key->GetDependentFunctions(DependentFunctions);
+			for (auto Function : DependentFunctions)
+			{
+				for (UMaterialExpression* FunctionExpression : Function->GetExpressions())
+				{
+					if (const UMaterialExpressionTextureSampleParameter* TexParameter = Cast<const UMaterialExpressionTextureSampleParameter>(FunctionExpression))
+					{
+						if (TexParameter->ParameterName == Parameter.Name)
+						{
+							FunctionParametersToVtIze.FindOrAdd(Function->GetBaseFunction()).AddUnique(Parameter);
+							InAffectedFunctions.AddUnique(Function->GetBaseFunction());
+							FunctionHeap.AddUnique(Function->GetBaseFunction());
+						}
+					}
+				}
+			}
+		}
+	}
+
+	// We now have a set of "root" functions which will be affected by changing InTextures to VT.
+	// Now find all children of these materials which will also be affected trough parameters now requiring VT textures being set on them.
+	// This will again load any child instances and their dependencies which aren't loaded yet
+	while (FunctionHeap.Num() > 0)
+	{
+		UMaterialFunctionInterface *ParentFunction = FunctionHeap[0];
+		FunctionHeap.RemoveAt(0);
+		{
+			UMaterialFunctionInstance* FunctionInstance = Cast<UMaterialFunctionInstance>(ParentFunction);
+			if (FunctionInstance)
+			{
+				// Check all parameters of the current material. If they reference a texture
+				// we want to convert to VT flag the parameter (this will then cause all textures assigned to this parameter to convert to vt as well)
+				for (const FTextureParameterValue& TextureParameter : FunctionInstance->TextureParameterValues)
+				{
+					UTexture2D* ParamValue2D = Cast<UTexture2D>(TextureParameter.ParameterValue);
+					if (InAffectedTextures.Contains(ParamValue2D))
+					{
+						FunctionParametersToVtIze.FindOrAdd(ParentFunction->GetBaseFunction()).Add(TextureParameter.ParameterInfo);
+					}
+				}
+			}
+		}
+
+		// Find all direct children of this function (children of children will be discovered later by pushing these children on the FunctionHeap).
+		TArray<UMaterialFunctionInstance*> ParentFunctionInstances;
+		Task.EnterProgressFrame();
+		GetReferencersOfType(ParentFunction, ParentFunctionInstances);
+
+		for (auto FunctionInstance : ParentFunctionInstances)
+		{
+			FunctionInferfaces.AddUnique(FunctionInstance);
+
+			// Push on the heap to check materials referencing us recursively
+			FunctionHeap.AddUnique(FunctionInstance);
+		}
+	}
+
+	// We know have a set of root functions and a set of properties to convert to VT
+	// Find all textures referenced by these properties
+	// These new textures could in turn be referenced by other materials and or functions (not in the inheritance chain)
+	// which is why we have to run this discovery process iteratively until we don't discover any new
+	// materials or textures.
+
+	for (UMaterialFunctionInterface *If : FunctionInferfaces)
+	{
+		UMaterialFunctionInterface *Func = If->GetBaseFunction();
+		for (const FMaterialParameterInfo &Parameter : FunctionParametersToVtIze.FindOrAdd(Func))
+		{
+			UTexture *Tex = nullptr;
+			If->OverrideNamedTextureParameter(Parameter, Tex);
+			UTexture2D *Tex2D = Cast<UTexture2D>(Tex);
+			if (Tex2D && !Tex2D->VirtualTextureStreaming)
+			{
+				InAffectedTextures.AddUnique(Tex2D);
+				AuditTrail.Add(Tex2D, FAuditTrail(
+					Func,
+					FString::Printf(TEXT("set on parameter %s in instance %s"), *Parameter.Name.ToString(), *If->GetPathName())
+				));
+			}
+		}
+	}
+}
+
+void FVirtualTextureConversionWorker::FindAllTexturesAndMaterials(TArray<UMaterial *> &OutAffectedMaterials, TArray<UMaterialFunctionInterface *> &OutAffectedFunctions, TArray<UTexture2D *> &OutAffectedTextures)
+{
+	int LastNumMaterials = OutAffectedMaterials.Num();
+	int LastNumTextures = OutAffectedTextures.Num();
+	int LastNumFunctions = OutAffectedFunctions.Num();
+
+	FScopedSlowTask SlowTask(1000.0f, LOCTEXT("ConvertToVT_Progress_FindAllTexturesAndMaterials", "Finding Textures and Materials..."));
+
+	while (true)
+	{
+		FindAllTexturesAndMaterials_Iteration(OutAffectedMaterials, OutAffectedFunctions, OutAffectedTextures, MaterialRejectedTextures, SlowTask);
+		if (OutAffectedMaterials.Num() == LastNumMaterials && OutAffectedTextures.Num() == LastNumTextures && OutAffectedFunctions.Num() == LastNumFunctions)
+		{
+			return;
+		}
+
+		LastNumMaterials = OutAffectedMaterials.Num();
+		LastNumTextures = OutAffectedTextures.Num();
+		LastNumFunctions = OutAffectedFunctions.Num();
+	}
+}
+
+void FVirtualTextureConversionWorker::FilterList(int32 SizeThreshold)
+{
+	FScopedSlowTask SlowTask(1.0f, LOCTEXT("ConvertToVT_Progress_FindingTextures", "Finding textures to convert..."));
+	SlowTask.MakeDialog();
+
+	Textures.Empty();
+	Materials.Empty();
+	Functions.Empty();
+	SizeRejectedTextures.Empty();
+
+	for (UTexture2D *Texture : UserTextures)
+	{
+		if (Texture->GetSizeX()*Texture->GetSizeY() >= SizeThreshold * SizeThreshold)
+		{
+			Textures.Add(Texture);
+		}
+		else
+		{
+			SizeRejectedTextures.Add(Texture);
+		}
+	}
+
+	SlowTask.EnterProgressFrame();
+	FindAllTexturesAndMaterials(Materials, Functions, Textures);
+}
+
+static void MarkTextureExpressionModified(UMaterialExpressionTextureBase* Expression)
+{
+	FPropertyChangedEvent SamplerTypeChangeEvent(UMaterialExpressionTextureBase::StaticClass()->FindPropertyByName(GET_MEMBER_NAME_CHECKED(UMaterialExpressionTextureBase, SamplerType)));
+	FPropertyChangedEvent TextureChangeEvent(UMaterialExpressionTextureBase::StaticClass()->FindPropertyByName(GET_MEMBER_NAME_CHECKED(UMaterialExpressionTextureBase, Texture)));
+
+	Expression->Modify();
+
+	// Nofity that we changed SamplerType
+	Expression->PostEditChangeProperty(SamplerTypeChangeEvent);
+
+	// Also notify that we changed Texture (technically didn't modify this property)
+	// This way code in FMaterialExpressionTextureBaseDetails will see the event, and refresh the list of valid sampler types for the updated texture
+	Expression->PostEditChangeProperty(TextureChangeEvent);
+}
+
+void FVirtualTextureConversionWorker::DoConvert()
+{
+	const bool bVirtualTextureEnable = !bConvertBackward;
+
+	FScopedSlowTask SlowTask(2.0f, LOCTEXT("ConvertToVT_Progress_ConvertingTexturesAndMaterials", "Converting textures and materials..."));
+	SlowTask.MakeDialog();
+
+	TMap<UTexture2D*, UTexture2D*> EngineTextureToCopyMap;
+
+	UE_LOG(LogVirtualTextureConversion, Display, TEXT("Beginning conversion..."));
+	SlowTask.EnterProgressFrame();
+	{
+		FScopedSlowTask TextureTask(Textures.Num(), LOCTEXT("ConvertToVT_Progress_TextureTask", "Updating textures..."));
+
+		for (UTexture2D *Tex : Textures)
+		{
+			UTexture2D* TextureToUpdate = Tex;
+			if (TextureToUpdate->GetPathName().StartsWith(TEXT("/Engine/")) && TextureToUpdate->GetPackage() != GetTransientPackage())
+			{
+				// rather than modify engine content, create a copy and update that
+				// any materials that we modify will be updated to point to the copy
+				ObjectTools::FPackageGroupName PGN;
+				PGN.GroupName = TEXT("");
+				PGN.ObjectName = Tex->GetName().Append(TEXT("_VT"));
+				PGN.PackageName = TEXT("/Game/Textures/");
+				PGN.PackageName.Append(PGN.ObjectName);
+
+				UPackage* ExistingPackage = FindPackage(NULL, *PGN.PackageName);
+				UTexture2D* DuplicateTexture = nullptr;
+				if (ExistingPackage)
+				{
+					UObject* DuplicateObject = StaticFindObject(UTexture2D::StaticClass(), ExistingPackage, *PGN.ObjectName);
+					if (DuplicateObject)
+					{
+						DuplicateTexture = CastChecked<UTexture2D>(DuplicateObject);
+					}
+				}
+
+				if (!DuplicateTexture || DuplicateTexture->VirtualTextureStreaming != bVirtualTextureEnable)
+				{
+					// TODO - overwrite previous texture (if it exists), or should we always generate a unique name?
+					TSet<UPackage*> ObjectsUserRefusedToFullyLoad;
+					UObject* DuplicateObject = ObjectTools::DuplicateSingleObject(Tex, PGN, ObjectsUserRefusedToFullyLoad, false);
+					if (DuplicateObject)
+					{
+						DuplicateTexture = CastChecked<UTexture2D>(DuplicateObject);
+					}
+				}
+
+				TextureToUpdate = DuplicateTexture;
+				EngineTextureToCopyMap.Add(Tex, TextureToUpdate);
+				if (!TextureToUpdate)
+				{
+					UE_LOG(LogVirtualTextureConversion, Warning, TEXT("Failed to duplicate engine texture %s"), *Tex->GetPathName());
+				}
+			}
+
+			if (TextureToUpdate)
+			{
+				UE_LOG(LogVirtualTextureConversion, Display, TEXT("Texture %s"), *TextureToUpdate->GetName());
+				TextureTask.EnterProgressFrame();
+
+				if (TextureToUpdate->VirtualTextureStreaming != bVirtualTextureEnable)
+				{
+					FPropertyChangedEvent PropertyChangeEvent(UTexture::StaticClass()->FindPropertyByName(GET_MEMBER_NAME_CHECKED(UTexture, VirtualTextureStreaming)));
+					TextureToUpdate->Modify();
+					TextureToUpdate->VirtualTextureStreaming = bVirtualTextureEnable;
+					TextureToUpdate->PostEditChangeProperty(PropertyChangeEvent);
+				}
+			}
+		}
+	}
+
+	SlowTask.EnterProgressFrame();
+	{
+		FScopedSlowTask MaterialTask(Materials.Num() + Functions.Num(), LOCTEXT("ConvertToVT_Progress_MaterialTask", "Updating materials..."));
+		FMaterialUpdateContext UpdateContext;
+
+		TMap<UMaterialFunctionInterface*, TArray<UMaterial*>> FunctionToMaterialMap;
+
+		for (UMaterial *Mat : Materials)
+		{
+			UE_LOG(LogVirtualTextureConversion, Display, TEXT("Material %s"), *Mat->GetName());
+
+			MaterialTask.EnterProgressFrame();
+
+			bool MatModified = false;
+			for (UMaterialExpression *Expr : Mat->GetExpressions())
+			{
+				UMaterialExpressionTextureBase *TexExpr = Cast<UMaterialExpressionTextureBase>(Expr);
+				if (TexExpr)
+				{
+					if (Textures.Contains(TexExpr->Texture))
+					{
+						UE_LOG(LogVirtualTextureConversion, Display, TEXT("Adjusting sampler %s."), *TexExpr->GetName());
+						UTexture2D** FoundTextureCopy = EngineTextureToCopyMap.Find(CastChecked<UTexture2D>(TexExpr->Texture));
+						UTexture2D* TextureCopy = nullptr;
+						bool bShouldUpdateMaterial = true;
+						if (FoundTextureCopy)
+						{
+							TextureCopy = *FoundTextureCopy;
+							if (TextureCopy)
+							{
+								TexExpr->Texture = TextureCopy;
+							}
+							else
+							{
+								// nullptr was set in EngineTextureToCopyMap....this means we failed to create copy of engine texture for this resource
+								// bail on updating the material in this case
+								bShouldUpdateMaterial = false;
+							}
+						}
+
+						auto OldType = TexExpr->SamplerType;
+						if (bShouldUpdateMaterial)
+						{
+							TexExpr->AutoSetSampleType();
+						}
+						if (TextureCopy != nullptr || OldType != TexExpr->SamplerType)
+						{
+							MarkTextureExpressionModified(TexExpr);
+							MatModified = true;
+						}
+					}
+				}
+			}
+
+			TArray<UMaterialFunctionInterface *>MaterialFunctions;
+			Mat->GetDependentFunctions(MaterialFunctions);
+			for (auto Function : MaterialFunctions)
+			{
+				FunctionToMaterialMap.FindOrAdd(Function).AddUnique(Mat);
+			}
+
+			if (MatModified)
+			{
+				UE_LOG(LogVirtualTextureConversion, Display, TEXT("Material %s added to update list."), *Mat->GetName());
+				UMaterialEditingLibrary::RecompileMaterial(Mat);
+			}
+			else
+			{
+				UE_LOG(LogVirtualTextureConversion, Display, TEXT("Material %s was not modified, skipping."), *Mat->GetName());
+			}
+		}
+
+		for (UMaterialFunctionInterface *Func : Functions)
+		{
+			UE_LOG(LogVirtualTextureConversion, Display, TEXT("Function %s"), *Func->GetName());
+
+			MaterialTask.EnterProgressFrame();
+
+			bool FuncModified = false;
+			for (const TObjectPtr<UMaterialExpression>& Expr : Func->GetExpressions())
+			{
+				UMaterialExpressionTextureBase *TexExpr = Cast<UMaterialExpressionTextureBase>(Expr);
+				if (TexExpr)
+				{
+					if (Textures.Contains(TexExpr->Texture))
+					{
+						UE_LOG(LogVirtualTextureConversion, Display, TEXT("Adjusting sampler %s."), *TexExpr->GetName());
+						UTexture2D** FoundTextureCopy = EngineTextureToCopyMap.Find(CastChecked<UTexture2D>(TexExpr->Texture));
+						UTexture2D* TextureCopy = nullptr;
+						bool bShouldUpdateMaterial = true;
+						if (FoundTextureCopy)
+						{
+							TextureCopy = *FoundTextureCopy;
+							if (TextureCopy)
+							{
+								TexExpr->Texture = TextureCopy;
+							}
+							else
+							{
+								// nullptr was set in EngineTextureToCopyMap....this means we failed to create copy of engine texture for this resource
+								// bail on updating the material in this case
+								bShouldUpdateMaterial = false;
+							}
+						}
+
+						auto OldType = TexExpr->SamplerType;
+						if (bShouldUpdateMaterial)
+						{
+							TexExpr->AutoSetSampleType();
+						}
+						if (TextureCopy != nullptr || TexExpr->SamplerType != OldType)
+						{
+							MarkTextureExpressionModified(TexExpr);
+							FuncModified = true;
+						}
+					}
+				}
+			}
+
+			if (FuncModified)
+			{
+				UMaterialEditingLibrary::UpdateMaterialFunction(Func, nullptr);
+				UE_LOG(LogVirtualTextureConversion, Display, TEXT("Function %s added to update list."), *Func->GetName());
+			}
+			else
+			{
+				UE_LOG(LogVirtualTextureConversion, Display, TEXT("Function %s was not modified, skipping."), *Func->GetName());
+			}
+		}
+
+		// update the world's viewports
+		UE_LOG(LogVirtualTextureConversion, Display, TEXT("Broadcasting to editor."));
+		FEditorDelegates::RefreshEditor.Broadcast();
+		FEditorSupportDelegates::RedrawAllViewports.Broadcast();
+	}
+}
+
+#undef LOCTEXT_NAMESPACE
\ No newline at end of file
diff --git a/Engine/Source/Editor/VirtualTexturingEditor/Private/VirtualTextureConversionWorker.h b/Engine/Source/Editor/VirtualTexturingEditor/Private/VirtualTextureConversionWorker.h
new file mode 100644
--- /dev/null
+++ b/Engine/Source/Editor/VirtualTexturingEditor/Private/VirtualTextureConversionWorker.h
@@ -0,0 +1,65 @@
+// Copyright Epic Games, Inc. All Rights Reserved.
+
+#pragma once
+
+#include "CoreMinimal.h"
+#include "AssetRegistry/AssetRegistryModule.h"
+#include "UObject/GCObject.h"
+#include "AssetRegistry/AssetIdentifier.h"
+
+class UTexture2D;
+class UMaterial;
+class UMaterialFunctionInterface;
+struct FScopedSlowTask;
+class UClass;
+struct FAssetData;
+
+class FVirtualTextureConversionWorker : public FGCObject
+{
+public:
+	FVirtualTextureConversionWorker(bool bInConvertBackward = false) : bConvertBackward(bInConvertBackward){}
+
+	bool bConvertBackward;
+
+	// Fill this array intially with some textures you want to see converted to vt
+	TArray<UTexture2D*> UserTextures;
+
+	// These will be filled by calling FilterList
+	TArray<UTexture2D*> Textures;
+	TArray<UMaterial *>Materials;
+	TArray<UMaterialFunctionInterface *> Functions;
+	TArray<UTexture2D*> SizeRejectedTextures;// The original selection of the user filtered by textures that match the threshold size
+	TArray<UTexture2D*> MaterialRejectedTextures; // Textures that could not be converted due to their usage in materials (connected to unsupported property)
+
+	struct FAuditTrail
+	{
+		FAuditTrail(UObject *SetObject, const FString &SetString) : Destination(SetObject), PathDescription(SetString) {}
+		UObject *Destination; //Next object in the trail
+		FString PathDescription; //Description on how we go there
+	};
+
+	TMap<UObject*, FAuditTrail> AuditTrail;
+
+	void SetConversionDirection(bool bInConvertBackward) { bConvertBackward = bInConvertBackward; }
+
+	void FilterList(int32 SizeThreshold);
+
+	// Based on the filtered lists do the actual VT conversion
+	void DoConvert();
+
+	void AddReferencedObjects( FReferenceCollector& Collector ) override;
+	virtual FString GetReferencerName() const override
+	{
+		return TEXT("FVTConversionWorker");
+	}
+
+private:
+	void FindAllTexturesAndMaterials_Iteration(TArray<UMaterial*>& InAffectedMaterials,
+		TArray<UMaterialFunctionInterface*>& InAffectedFunctions,
+		TArray<UTexture2D*>& InAffectedTextures,
+		TArray<UTexture2D*>& InInvalidTextures,
+		FScopedSlowTask& Task);
+
+	void FindAllTexturesAndMaterials(TArray<UMaterial *> &OutAffectedMaterials, TArray<UMaterialFunctionInterface *> &OutAffectedFunctions, TArray<UTexture2D *> &OutAffectedTextures);
+
+};
diff --git a/Engine/Source/Editor/VirtualTexturingEditor/Private/VirtualTexturingEditorModule.cpp b/Engine/Source/Editor/VirtualTexturingEditor/Private/VirtualTexturingEditorModule.cpp
--- a/Engine/Source/Editor/VirtualTexturingEditor/Private/VirtualTexturingEditorModule.cpp
+++ b/Engine/Source/Editor/VirtualTexturingEditor/Private/VirtualTexturingEditorModule.cpp
@@ -8,12 +8,14 @@
 #include "RuntimeVirtualTextureBuildStreamingMips.h"
 #include "RuntimeVirtualTextureDetailsCustomization.h"
 #include "RuntimeVirtualTextureThumbnailRenderer.h"
+#include "SConvertToVirtualTexture.h"
 #include "ThumbnailRendering/ThumbnailManager.h"
 #include "VirtualTextureBuilderAssetTypeActions.h"
 #include "VirtualTextureBuilderThumbnailRenderer.h"
 #include "VT/RuntimeVirtualTexture.h"
 #include "VT/RuntimeVirtualTextureVolume.h"
 #include "VT/VirtualTextureBuilder.h"
+#include "VirtualTextureConversionWorker.h"
 
 #define LOCTEXT_NAMESPACE "VirtualTexturingEditorModule"
 
@@ -30,6 +32,8 @@ public:
 	//~ Begin IVirtualTexturingEditorModule Interface.
 	virtual bool HasStreamedMips(URuntimeVirtualTextureComponent* InComponent) const override;
 	virtual bool BuildStreamedMips(URuntimeVirtualTextureComponent* InComponent) const override;
+	virtual void ConvertVirtualTextures(const TArray<UTexture2D *>& Textures, bool bConvertBackToNonVirtual, const TArray<UMaterial *>* RelatedMaterials /* = nullptr */) const override;
+	virtual void ConvertVirtualTexturesWithDialog(const TArray<UTexture2D *>& Textures, bool bConvertBackToNonVirtual) const override;
 	//~ End IVirtualTexturingEditorModule Interface.
 
 private:
@@ -88,4 +92,23 @@ bool FVirtualTexturingEditorModule::BuildStreamedMips(URuntimeVirtualTextureComp
 	return RuntimeVirtualTexture::BuildStreamedMips(InComponent, ERuntimeVirtualTextureDebugType::None);
 }
 
+void FVirtualTexturingEditorModule::ConvertVirtualTextures(const TArray<UTexture2D *>& Textures, bool bConvertBackToNonVirtual, const TArray<UMaterial *>* RelatedMaterials /* = nullptr */) const
+{
+	FVirtualTextureConversionWorker VirtualTextureConversionWorker(bConvertBackToNonVirtual);
+	VirtualTextureConversionWorker.UserTextures = Textures;
+	//We want all given texture to be added, so we put a minimum texture size of 0
+	VirtualTextureConversionWorker.FilterList(0);
+	if (RelatedMaterials)
+	{
+		VirtualTextureConversionWorker.Materials.Append(*RelatedMaterials);
+	}
+
+	VirtualTextureConversionWorker.DoConvert();
+}
+
+void FVirtualTexturingEditorModule::ConvertVirtualTexturesWithDialog(const TArray<UTexture2D*>& Textures, bool bConvertBackToNonVirtual) const
+{
+	SConvertToVirtualTexture::ConvertVTTexture(Textures, bConvertBackToNonVirtual);
+}
+
 #undef LOCTEXT_NAMESPACE
diff --git a/Engine/Source/Editor/VirtualTexturingEditor/Public/VirtualTexturingEditorModule.h b/Engine/Source/Editor/VirtualTexturingEditor/Public/VirtualTexturingEditorModule.h
--- a/Engine/Source/Editor/VirtualTexturingEditor/Public/VirtualTexturingEditorModule.h
+++ b/Engine/Source/Editor/VirtualTexturingEditor/Public/VirtualTexturingEditorModule.h
@@ -5,6 +5,9 @@
 #include "CoreMinimal.h"
 #include "Modules/ModuleInterface.h"
 
+class UTexture2D;
+class UMaterial;
+
 /** Module for virtual texturing editor extensions. */
 class IVirtualTexturingEditorModule : public IModuleInterface
 {
@@ -13,4 +16,7 @@ public:
 	virtual bool HasStreamedMips(class URuntimeVirtualTextureComponent* InComponent) const = 0;
 	/** Build the contents of the streaming low mips. */
 	virtual bool BuildStreamedMips(class URuntimeVirtualTextureComponent* InComponent) const = 0;
+	
+	virtual void ConvertVirtualTextures(const TArray<UTexture2D *>& Textures, bool bConvertBackToNonVirtual, const TArray<UMaterial *>* RelatedMaterials /* = nullptr */) const = 0;
+	virtual void ConvertVirtualTexturesWithDialog(const TArray<UTexture2D*>& Textures, bool bConvertBackToNonVirtual) const = 0;
 };
diff --git a/Engine/Source/Runtime/AssetRegistry/Private/AssetRegistryHelpers.cpp b/Engine/Source/Runtime/AssetRegistry/Private/AssetRegistryHelpers.cpp
--- a/Engine/Source/Runtime/AssetRegistry/Private/AssetRegistryHelpers.cpp
+++ b/Engine/Source/Runtime/AssetRegistry/Private/AssetRegistryHelpers.cpp
@@ -4,6 +4,7 @@
 
 #include "AssetRegistry/IAssetRegistry.h"
 #include "AssetRegistry.h"
+#include "Algo/AnyOf.h"
 #include "Blueprint/BlueprintSupport.h"
 #include "Misc/PackageName.h"
 #include "Modules/ModuleManager.h"
@@ -113,6 +114,41 @@ UClass* UAssetRegistryHelpers::FindAssetNativeClass(const FAssetData& AssetData)
 	return AssetClass;
 }
 
+void UAssetRegistryHelpers::FindReferencersOfAssetOfClass(UObject* AssetInstance, TConstArrayView<UClass*> InMatchClasses, TArray<FAssetData>& OutAssetDatas)
+{
+	FindReferencersOfAssetOfClass(AssetInstance->GetOutermost()->GetFName(), InMatchClasses, OutAssetDatas);
+}
+
+void UAssetRegistryHelpers::FindReferencersOfAssetOfClass(const FAssetIdentifier& InAssetIdentifier, TConstArrayView<UClass*> InMatchClasses, TArray<FAssetData>& OutAssetDatas)
+{
+	// If the asset registry is still loading assets, we cant check for referencers, so we must open the rename dialog
+	const IAssetRegistry& AssetRegistry = IAssetRegistry::GetChecked();
+
+	TArray<FAssetIdentifier> Referencers;
+	AssetRegistry.GetReferencers(InAssetIdentifier, Referencers);
+
+	for (auto AssetIdentifier : Referencers)
+	{
+		TArray<FAssetData> Assets;
+		AssetRegistry.GetAssetsByPackageName(AssetIdentifier.PackageName, Assets);
+
+		for (auto AssetData : Assets)
+		{
+			if (InMatchClasses.Num() > 0)
+			{
+				if (Algo::AnyOf(InMatchClasses, [&AssetData](UClass* MatchClass){ return AssetData.IsInstanceOf(MatchClass); }))
+				{
+					OutAssetDatas.AddUnique(AssetData);
+				}
+			}
+			else
+			{
+				OutAssetDatas.AddUnique(AssetData);
+			}
+		}
+	}
+}
+
 void UAssetRegistryHelpers::GetBlueprintAssets(const FARFilter& InFilter, TArray<FAssetData>& OutAssetData)
 {
 	IAssetRegistry& AssetRegistry = IAssetRegistry::GetChecked();
diff --git a/Engine/Source/Runtime/AssetRegistry/Public/AssetRegistry/AssetRegistryHelpers.h b/Engine/Source/Runtime/AssetRegistry/Public/AssetRegistry/AssetRegistryHelpers.h
--- a/Engine/Source/Runtime/AssetRegistry/Public/AssetRegistry/AssetRegistryHelpers.h
+++ b/Engine/Source/Runtime/AssetRegistry/Public/AssetRegistry/AssetRegistryHelpers.h
@@ -15,12 +15,14 @@
 #include "UObject/ScriptInterface.h"
 #include "UObject/SoftObjectPath.h"
 #include "UObject/UObjectGlobals.h"
+#include "Containers/ArrayView.h"
 
 #include "AssetRegistryHelpers.generated.h"
 
 class IAssetRegistry;
 class UClass;
 struct FFrame;
+struct FAssetIdentifier;
 
 USTRUCT(BlueprintType)
 struct FTagAndValue
@@ -108,6 +110,16 @@ public:
 	 */
 	UFUNCTION(BlueprintPure, Category = "Asset Registry", meta=(ScriptMethod))
 	static ASSETREGISTRY_API UClass* FindAssetNativeClass(const FAssetData& AssetData);
+
+	/**
+	 * Finds references of the provided asset that are of the a class contained in the InMatchClasses set.
+	 */
+	static ASSETREGISTRY_API void FindReferencersOfAssetOfClass(UObject* AssetInstance, TConstArrayView<UClass*> InMatchClasses, TArray<FAssetData>& OutAssetDatas);
+
+	/**
+	 * Finds references of the provided asset that are of the a class contained in the InMatchClasses set.
+	 */
+	static ASSETREGISTRY_API void FindReferencersOfAssetOfClass(const FAssetIdentifier& InAssetIdentifier, TConstArrayView<UClass*> InMatchClasses, TArray<FAssetData>& OutAssetDatas);
 	
 	/** Enable/disable asset registry caching mode for the duration of the scope */
 	struct ASSETREGISTRY_API FTemporaryCachingModeScope
