Fixing unity build issue

diff --git a/Engine/Source/Editor/VirtualTexturingEditor/Private/VirtualTextureConversionWorker.cpp b/Engine/Source/Editor/VirtualTexturingEditor/Private/VirtualTextureConversionWorker.cpp
--- a/Engine/Source/Editor/VirtualTexturingEditor/Private/VirtualTextureConversionWorker.cpp
+++ b/Engine/Source/Editor/VirtualTexturingEditor/Private/VirtualTextureConversionWorker.cpp
@@ -25,24 +25,27 @@
 
 DEFINE_LOG_CATEGORY_STATIC(LogVirtualTextureConversion, Log, All);
 
-template <class T> void GetReferencersOfType(UObject *Object, TArray<T*> &OutObjects)
+namespace VTConversionWorkerUtil
 {
-	TArray<FAssetData> OutAssetDatas;
-	UAssetRegistryHelpers::FindReferencersOfAssetOfClass(Object, { T::StaticClass() }, OutAssetDatas);
-
-	FScopedSlowTask SlowTask(OutAssetDatas.Num(), LOCTEXT("ConvertToVT_Progress_LoadingObjects", "Loading Objects..."));
-
-	for (auto Data : OutAssetDatas)
+	template <class T> void GetReferencersOfType(UObject *Object, TArray<T*> &OutObjects)
 	{
-		SlowTask.EnterProgressFrame();
-		UObject *MaybeOk = Data.GetAsset();
+		TArray<FAssetData> OutAssetDatas;
+		UAssetRegistryHelpers::FindReferencersOfAssetOfClass(Object, { T::StaticClass() }, OutAssetDatas);
 
-		// It's a material?
-		T *IsOk = Cast<T>(MaybeOk);
+		FScopedSlowTask SlowTask(OutAssetDatas.Num(), LOCTEXT("ConvertToVT_Progress_LoadingObjects", "Loading Objects..."));
 
-		if (IsOk != nullptr)
+		for (auto Data : OutAssetDatas)
 		{
-			OutObjects.Add(IsOk);
+			SlowTask.EnterProgressFrame();
+			UObject *MaybeOk = Data.GetAsset();
+
+			// It's a material?
+			T *IsOk = Cast<T>(MaybeOk);
+
+			if (IsOk != nullptr)
+			{
+				OutObjects.Add(IsOk);
+			}
 		}
 	}
 }
@@ -121,7 +124,7 @@ void FVirtualTextureConversionWorker::FindAllTexturesAndMaterials_Iteration(TArr
 
 		Task.EnterProgressFrame();
 		TArray<UMaterialInterface *> MaterialsUsingTexture;
-		GetReferencersOfType(Tex2D, MaterialsUsingTexture);
+		VTConversionWorkerUtil::GetReferencersOfType(Tex2D, MaterialsUsingTexture);
 
 		/*for (auto Material : MaterialsUsingTexture)
 		{
@@ -174,7 +177,7 @@ void FVirtualTextureConversionWorker::FindAllTexturesAndMaterials_Iteration(TArr
 	{
 		Task.EnterProgressFrame();
 		TArray<UMaterialInterface *> MaterialsUsingFunction;
-		GetReferencersOfType(Func, MaterialsUsingFunction);
+		VTConversionWorkerUtil::GetReferencersOfType(Func, MaterialsUsingFunction);
 
 		/*for (auto Material : MaterialsUsingFunction)
 		{
@@ -256,7 +259,7 @@ void FVirtualTextureConversionWorker::FindAllTexturesAndMaterials_Iteration(TArr
 		// Find all direct children of this material (children of children will be discovered later by pushing these children on the MaterialHeap).
 		TArray<UMaterialInstance*> ParentMaterialInstances;
 		Task.EnterProgressFrame();
-		GetReferencersOfType(ParentMaterial, ParentMaterialInstances);
+		VTConversionWorkerUtil::GetReferencersOfType(ParentMaterial, ParentMaterialInstances);
 
 		for (auto MaterialInstance : ParentMaterialInstances)
 		{
@@ -306,7 +309,7 @@ void FVirtualTextureConversionWorker::FindAllTexturesAndMaterials_Iteration(TArr
 		Task.EnterProgressFrame();
 		if (!Tex2D->GetPathName().StartsWith("/Engine/"))
 		{
-			GetReferencersOfType(Tex2D, FunctionsUsedByAffectedTextures);
+			VTConversionWorkerUtil::GetReferencersOfType(Tex2D, FunctionsUsedByAffectedTextures);
 		}
 	}
 
@@ -383,7 +386,7 @@ void FVirtualTextureConversionWorker::FindAllTexturesAndMaterials_Iteration(TArr
 		// Find all direct children of this function (children of children will be discovered later by pushing these children on the FunctionHeap).
 		TArray<UMaterialFunctionInstance*> ParentFunctionInstances;
 		Task.EnterProgressFrame();
-		GetReferencersOfType(ParentFunction, ParentFunctionInstances);
+		VTConversionWorkerUtil::GetReferencersOfType(ParentFunction, ParentFunctionInstances);
 
 		for (auto FunctionInstance : ParentFunctionInstances)
 		{
