More fixes for the convert to virtual texture dialog, it no longer needs to load every asset that references a texutre just to see if it's a material.

diff --git a/Engine/Source/Editor/VirtualTexturingEditor/Private/SConvertToVirtualTexture.cpp b/Engine/Source/Editor/VirtualTexturingEditor/Private/SConvertToVirtualTexture.cpp
--- a/Engine/Source/Editor/VirtualTexturingEditor/Private/SConvertToVirtualTexture.cpp
+++ b/Engine/Source/Editor/VirtualTexturingEditor/Private/SConvertToVirtualTexture.cpp
@@ -443,7 +443,7 @@ void SConvertToVirtualTexture::UpdateList()
 		Status->NonPowerOf2 = !Texture->Source.IsPowerOfTwo() && Texture->PowerOfTwoMode == ETexturePowerOfTwoSetting::None;
 	}
 
-	for (UTexture2D *Texture : Worker.MaterialRejectedTextures)
+	for (UTexture2D* Texture : Worker.MaterialRejectedTextures)
 	{
 		AssetList.Add(Texture);
 		FConversionStatus* Status = new(AssetStatus) FConversionStatus();
@@ -464,7 +464,7 @@ void SConvertToVirtualTexture::UpdateList()
 		}
 	}
 
-	for (UMaterialFunctionInterface *Func : Worker.Functions)
+	for (UMaterialFunctionInterface* Func : Worker.Functions)
 	{
 		if (Func->GetOutermost() != GetTransientPackage())
 		{
@@ -630,212 +630,20 @@ FConvertToVTDlg::EResult FConvertToVTDlg::ShowModal()
 }
 
 
-void SConvertToVirtualTexture::ConvertVTTexture(TArray<UTexture2D*> Objects, bool backwards)
+void SConvertToVirtualTexture::ConvertVTTexture(TArray<UTexture2D*> InTextures, bool backwards)
 {
-	TArray<FAssetData> AllRelevantMaterials;
 	TArray<UTexture2D*> UserTextures; // The original selection of the user
-
-	//FScopedSlowTask SlowTask(4.0f, LOCTEXT("ConvertToVT_Progress_ConvertToVt", "Converting textures to VT..."));
-	//SlowTask.MakeDialog();
-
-	// First simply assemble a list of materials to show the user.
-	//SlowTask.EnterProgressFrame();
-	for (auto ObjIt = Objects.CreateConstIterator(); ObjIt; ++ObjIt)
+	
+	for (UTexture2D* Texture : InTextures)
 	{
-		auto Object = (*ObjIt);
-		UTexture2D *Texture = Cast<UTexture2D>(Object);
 		if (Texture != nullptr && Texture->VirtualTextureStreaming == backwards)
 		{
-			UAssetRegistryHelpers::FindReferencersOfAssetOfClass(Object, { UMaterialInterface::StaticClass() }, AllRelevantMaterials);
 			UserTextures.Add(Texture);
 		}
 	}
 
 	FConvertToVTDlg ConvertDlg(UserTextures, backwards);
-	/*ConvertDlg.AssetList.Append(UserTextures);
-	ConvertDlg.AssetList.Append(AllRelevantMaterials);
-	ConvertDlg.AssetStatus.AddDefaulted(UserTextures.Num());
-	ConvertDlg.AssetStatus.AddDefaulted(AllRelevantMaterials.Num());
-	ConvertDlg.Message = LOCTEXT("ConvertToVTMaterialList", "The following textures & materials (and their dependencies) will be loaded by this operation. This could take a while.");
-	ConvertDlg.Threshold = virtualTextureAutoEnableThreshold;*/
-
-	if (ConvertDlg.ShowModal() == FConvertToVTDlg::Confirm)
-	{
-#if 0
-		//FVTConversionWorker Conversion;
-		//Conversion.UserTextures = UserTextures;
-		//Conversion.FilterList(ConvertDlg.Threshold);
-
-		/*FScopedSlowTask SlowTask(3.0f, LOCTEXT("ConvertToVT_Progress_ConvertToVt", "Converting textures to VT..."));
-		SlowTask.MakeDialog();
-
-		TArray<UTexture2D*> Textures; // All textures that should be converted
-		TArray<UTexture2D*> SizeRejectedTextures; // The original selection of the user filtered by textures that match the threshold size
-
-		for (UTexture2D *Texture: UserTextures)
-		{
-			if (Texture->GetSizeX()*Texture->GetSizeY() >= ConvertDlg.Threshold*ConvertDlg.Threshold)
-			{
-				Textures.Add(Texture);
-			}
-			else
-			{
-				SizeRejectedTextures.Add(Texture);
-			}
-		}
-
-		SlowTask.EnterProgressFrame();
-		TArray<UMaterial *>Materials;
-		TArray<UMaterialFunctionInterface *> Functions;
-//		FindAllTexturesAndMaterials(Materials, Functions, Textures);
-*/
-		// Show the dialog yet again warning the user...
-
-		/*FConvertToVTDlg ConvertDlg2;
-		ConvertDlg2.AssetList.Empty();
-		for (UTexture2D *Texture : Conversion.Textures)
-		{
-			ConvertDlg2.AssetList.Add(Texture);
-			ConvertDlg2.AssetStatus.Add(FConversionStatus(UserTextures.Contains(Texture), (Texture->GetSizeX()*Texture->GetSizeY() < ConvertDlg.Threshold*ConvertDlg.Threshold)));
-		}
-
-		for (UMaterial *RootMat : Conversion.Materials)
-		{
-			ConvertDlg2.AssetList.Add(RootMat);
-			ConvertDlg2.AssetStatus.Add(FConversionStatus());
-		}
-
-		for (UMaterialFunctionInterface *Func : Conversion.Functions)
-		{
-			ConvertDlg2.AssetList.Add(Func);
-			ConvertDlg2.AssetStatus.Add(FConversionStatus());
-		}
-
-		ConvertDlg2.Message = LOCTEXT("ConvertToVTTextureList", "The following textures and materials will converted. Please carefully inspect this list as textures not originally selected may also be included for technical reasons.");
-		ConvertDlg2.Threshold = ConvertDlg.Threshold;
-		ConvertDlg2.EditThreshold = false;
-		if (ConvertDlg2.ShowModal() == FConvertToVTDlg::Confirm)
-		{
-			Conversion.DoConvert(backwards);
-			/*UE_LOG(LogVirtualTextureConversion, Display, TEXT("Beginning conversion..."));
-			SlowTask.EnterProgressFrame();
-			{
-				FScopedSlowTask TextureTask(Textures.Num(), LOCTEXT("ConvertToVT_Progress_TextureTask", "Updating textures..."));
-
-				for (UTexture2D *Tex : Textures)
-				{
-					UE_LOG(LogVirtualTextureConversion, Display, TEXT("Texture %s"), *Tex->GetName());
-					TextureTask.EnterProgressFrame();
-
-					bool OldVt = Tex->VirtualTextureStreaming;
-					Tex->VirtualTextureStreaming = !backwards;
-					if (OldVt != Tex->VirtualTextureStreaming)
-					{
-						Tex->Modify();
-						Tex->PostEditChange();
-					}
-				}
-			}
-
-			SlowTask.EnterProgressFrame();
-			{
-				FScopedSlowTask MaterialTask(Materials.Num() + Functions.Num(), LOCTEXT("ConvertToVT_Progress_MaterialTask", "Updating materials..."));
-				FMaterialUpdateContext UpdateContext;
-
-				TMap<UMaterialFunctionInterface*, TArray<UMaterial*>> FunctionToMaterialMap;
-
-				for (UMaterial *Mat : Materials)
-				{
-					UE_LOG(LogVirtualTextureConversion, Display, TEXT("Material %s"), *Mat->GetName());
-
-					MaterialTask.EnterProgressFrame();
-
-					bool MatModified = false;
-					for (UMaterialExpression *Expr : Mat->Expressions)
-					{
-						UMaterialExpressionTextureBase *TexExpr = Cast<UMaterialExpressionTextureBase>(Expr);
-						if (TexExpr)
-						{
-							if (Textures.Contains(TexExpr->Texture))
-							{
-								UE_LOG(LogVirtualTextureConversion, Display, TEXT("Adjusting sampler %s."), *TexExpr->GetName());
-								auto OldType = TexExpr->SamplerType;
-								TexExpr->AutoSetSampleType();
-								if (TexExpr->SamplerType != OldType)
-								{
-									TexExpr->Modify();
-									TexExpr->PostEditChange();
-									MatModified = true;
-								}
-							}
-						}
-					}
-
-					TArray<UMaterialFunctionInterface *>MaterialFunctions;
-					Mat->GetDependentFunctions(MaterialFunctions);
-					for (auto Function : MaterialFunctions)
-					{
-						FunctionToMaterialMap.FindOrAdd(Function).AddUnique(Mat);
-					}
-
-					if (MatModified)
-					{
-						UE_LOG(LogVirtualTextureConversion, Display, TEXT("Material %s added to update list."), *Mat->GetName());
-						UMaterialEditingLibrary::RecompileMaterial(Mat);
-					}
-					else
-					{
-						UE_LOG(LogVirtualTextureConversion, Display, TEXT("Material %s was not modified, skipping."), *Mat->GetName());
-					}
-				}
-
-				for (UMaterialFunctionInterface *Func : Functions)
-				{
-					UE_LOG(LogVirtualTextureConversion, Display, TEXT("Function %s"), *Func->GetName());
-
-					MaterialTask.EnterProgressFrame();
-
-					bool FuncModified = false;
-					const TArray<UMaterialExpression*> *Expressions = Func->GetExpressions();
-					for (UMaterialExpression *Expr : *Expressions)
-					{
-						UMaterialExpressionTextureBase *TexExpr = Cast<UMaterialExpressionTextureBase>(Expr);
-						if (TexExpr)
-						{
-							if (Textures.Contains(TexExpr->Texture))
-							{
-								UE_LOG(LogVirtualTextureConversion, Display, TEXT("Adjusting sampler %s."), *TexExpr->GetName());
-								auto OldType = TexExpr->SamplerType;
-								TexExpr->AutoSetSampleType();
-								if (TexExpr->SamplerType != OldType)
-								{
-									TexExpr->Modify();
-									TexExpr->PostEditChange();
-									FuncModified = true;
-								}
-							}
-						}
-					}
-
-					if (FuncModified)
-					{
-						UMaterialEditingLibrary::UpdateMaterialFunction(Func, nullptr);
-						UE_LOG(LogVirtualTextureConversion, Display, TEXT("Function %s added to update list."), *Func->GetName());
-					}
-					else
-					{
-						UE_LOG(LogVirtualTextureConversion, Display, TEXT("Function %s was not modified, skipping."), *Func->GetName());
-					}
-				}
-
-				// update the world's viewports
-				UE_LOG(LogVirtualTextureConversion, Display, TEXT("Broadcasting to editor."));
-				FEditorDelegates::RefreshEditor.Broadcast();
-				FEditorSupportDelegates::RedrawAllViewports.Broadcast();
-			}
-		}*/
-		#endif
-	}
+	ConvertDlg.ShowModal();
 }
 
 #undef LOCTEXT_NAMESPACE
\ No newline at end of file
diff --git a/Engine/Source/Editor/VirtualTexturingEditor/Private/VirtualTextureConversionWorker.cpp b/Engine/Source/Editor/VirtualTexturingEditor/Private/VirtualTextureConversionWorker.cpp
--- a/Engine/Source/Editor/VirtualTexturingEditor/Private/VirtualTextureConversionWorker.cpp
+++ b/Engine/Source/Editor/VirtualTexturingEditor/Private/VirtualTextureConversionWorker.cpp
@@ -37,14 +37,13 @@ namespace VTConversionWorkerUtil
 		for (auto Data : OutAssetDatas)
 		{
 			SlowTask.EnterProgressFrame();
-			UObject *MaybeOk = Data.GetAsset();
-
-			// It's a material?
-			T *IsOk = Cast<T>(MaybeOk);
-
-			if (IsOk != nullptr)
+			if (Data.IsInstanceOf(T::StaticClass()))
 			{
-				OutObjects.Add(IsOk);
+				T *IsOk = Cast<T>(Data.GetAsset());
+				if (IsOk != nullptr)
+				{
+					OutObjects.Add(IsOk);
+				}
 			}
 		}
 	}
