Recreate render proxies prior to rendering virtual textures in order to recreate cached mesh commands that get invalidated.

diff --git a/Engine/Source/Editor/VirtualTexturingEditor/Private/RuntimeVirtualTextureThumbnailRenderer.cpp b/Engine/Source/Editor/VirtualTexturingEditor/Private/RuntimeVirtualTextureThumbnailRenderer.cpp
--- a/Engine/Source/Editor/VirtualTexturingEditor/Private/RuntimeVirtualTextureThumbnailRenderer.cpp
+++ b/Engine/Source/Editor/VirtualTexturingEditor/Private/RuntimeVirtualTextureThumbnailRenderer.cpp
@@ -79,6 +79,11 @@ void URuntimeVirtualTextureThumbnailRenderer::Draw(UObject* Object, int32 X, int
 	FSceneInterface* Scene = RuntimeVirtualTextureComponent != nullptr ? RuntimeVirtualTextureComponent->GetScene() : nullptr;
 	check(Scene != nullptr);
 
+	if (UWorld* World = Scene->GetWorld())
+	{
+		World->SendAllEndOfFrameUpdates();
+	}
+
 	const FBox2D DestBox = FBox2D(FVector2D(X, Y), FVector2D(Width, Height));
 	const FTransform Transform = RuntimeVirtualTextureComponent->GetComponentTransform();
 	const FBox Bounds = RuntimeVirtualTextureComponent->Bounds.GetBox();
