Fixed non-unity build break.

diff --git a/Engine/Source/Editor/VirtualTexturingEditor/Private/RuntimeVirtualTextureThumbnailRenderer.cpp b/Engine/Source/Editor/VirtualTexturingEditor/Private/RuntimeVirtualTextureThumbnailRenderer.cpp
--- a/Engine/Source/Editor/VirtualTexturingEditor/Private/RuntimeVirtualTextureThumbnailRenderer.cpp
+++ b/Engine/Source/Editor/VirtualTexturingEditor/Private/RuntimeVirtualTextureThumbnailRenderer.cpp
@@ -3,6 +3,7 @@
 #include "RuntimeVirtualTextureThumbnailRenderer.h"
 
 #include "Components/RuntimeVirtualTextureComponent.h"
+#include "Engine/World.h"
 #include "Materials/MaterialRenderProxy.h"
 #include "MaterialShared.h"
 #include "Math/Box.h"
