Fixed: interface for backward compat from CL# 24335980

diff --git a/Engine/Source/Editor/VirtualTexturingEditor/Private/VirtualTexturingEditorModule.cpp b/Engine/Source/Editor/VirtualTexturingEditor/Private/VirtualTexturingEditorModule.cpp
--- a/Engine/Source/Editor/VirtualTexturingEditor/Private/VirtualTexturingEditorModule.cpp
+++ b/Engine/Source/Editor/VirtualTexturingEditor/Private/VirtualTexturingEditorModule.cpp
@@ -30,6 +30,8 @@ public:
 	//~ End IModuleInterface Interface.
 
 	//~ Begin IVirtualTexturingEditorModule Interface.
+	virtual bool HasStreamedMips(URuntimeVirtualTextureComponent* InComponent) const override;
+	virtual bool BuildStreamedMips(URuntimeVirtualTextureComponent* InComponent) const override;
 	virtual bool HasStreamedMips(EShadingPath ShadingPath, URuntimeVirtualTextureComponent* InComponent) const override;
 	virtual bool BuildStreamedMips(EShadingPath ShadingPath, URuntimeVirtualTextureComponent* InComponent) const override;
 	virtual void ConvertVirtualTextures(const TArray<UTexture2D *>& Textures, bool bConvertBackToNonVirtual, const TArray<UMaterial *>* RelatedMaterials /* = nullptr */) const override;
@@ -82,6 +84,16 @@ void FVirtualTexturingEditorModule::OnPlacementModeRefresh(FName CategoryName)
 	}
 }
 
+bool FVirtualTexturingEditorModule::HasStreamedMips(URuntimeVirtualTextureComponent* InComponent) const
+{
+	return RuntimeVirtualTexture::HasStreamedMips(EShadingPath::Deferred, InComponent);
+}
+
+bool FVirtualTexturingEditorModule::BuildStreamedMips(URuntimeVirtualTextureComponent* InComponent) const
+{
+	return RuntimeVirtualTexture::BuildStreamedMips(EShadingPath::Deferred, InComponent, ERuntimeVirtualTextureDebugType::None);
+}
+
 bool FVirtualTexturingEditorModule::HasStreamedMips(EShadingPath ShadingPath, URuntimeVirtualTextureComponent* InComponent) const
 {
 	return RuntimeVirtualTexture::HasStreamedMips(ShadingPath, InComponent);
diff --git a/Engine/Source/Editor/VirtualTexturingEditor/Public/VirtualTexturingEditorModule.h b/Engine/Source/Editor/VirtualTexturingEditor/Public/VirtualTexturingEditorModule.h
--- a/Engine/Source/Editor/VirtualTexturingEditor/Public/VirtualTexturingEditorModule.h
+++ b/Engine/Source/Editor/VirtualTexturingEditor/Public/VirtualTexturingEditorModule.h
@@ -14,10 +14,12 @@ class IVirtualTexturingEditorModule : public IModuleInterface
 {
 public:
 	/** Returns true if the component describes a runtime virtual texture that has streaming low mips. */
+	virtual bool HasStreamedMips(URuntimeVirtualTextureComponent* InComponent) const = 0;
 	virtual bool HasStreamedMips(EShadingPath ShadingPath, class URuntimeVirtualTextureComponent* InComponent) const = 0;
 	/** Build the contents of the streaming low mips. */
+	virtual bool BuildStreamedMips(URuntimeVirtualTextureComponent* InComponent) const = 0;
 	virtual bool BuildStreamedMips(EShadingPath ShadingPath, class URuntimeVirtualTextureComponent* InComponent) const = 0;
-	
+		
 	virtual void ConvertVirtualTextures(const TArray<UTexture2D *>& Textures, bool bConvertBackToNonVirtual, const TArray<UMaterial *>* RelatedMaterials /* = nullptr */) const = 0;
 	virtual void ConvertVirtualTexturesWithDialog(const TArray<UTexture2D*>& Textures, bool bConvertBackToNonVirtual) const = 0;
 };
