FindAllTexturesAndMaterials_Iteration: Optimize by changing array to set and by avoiding duplicated work.
#rnx

diff --git a/Engine/Source/Editor/VirtualTexturingEditor/Private/VirtualTextureConversionWorker.cpp b/Engine/Source/Editor/VirtualTexturingEditor/Private/VirtualTextureConversionWorker.cpp
--- a/Engine/Source/Editor/VirtualTexturingEditor/Private/VirtualTextureConversionWorker.cpp
+++ b/Engine/Source/Editor/VirtualTexturingEditor/Private/VirtualTextureConversionWorker.cpp
@@ -105,26 +105,25 @@ void FVirtualTextureConversionWorker::AddReferencedObjects( FReferenceCollector&
 	Collector.AddReferencedObjects( AuditTrailKeys );
 }
 
-void FVirtualTextureConversionWorker::FindAllTexturesAndMaterials_Iteration(TArray<UMaterial*>& InAffectedMaterials,
-	TArray<UMaterialFunctionInterface*>& InAffectedFunctions,
-	TArray<UTexture2D*>& InAffectedTextures,
-	TArray<UTexture2D*>& InInvalidTextures,
+void FVirtualTextureConversionWorker::FindAllTexturesAndMaterials_Iteration(TSet<UMaterial*>& InAffectedMaterials,
+	TSet<UMaterialFunctionInterface*>& InAffectedFunctions,
+	TSet<UTexture2D*>& InAffectedTextures,
+	TSet<UTexture2D*>& InInvalidTextures,
 	FScopedSlowTask& Task)
 {
 	TArray<UMaterialInterface *>MaterialInferfaces; // All parents and instances
 	TArray<UMaterialInterface *>MaterialHeap; // Temporary work heap
 	TArray<UMaterialFunctionInterface *>FunctionInferfaces; // All parents and instances
 	TArray<UMaterialFunctionInterface *>FunctionHeap; // Temporary work heap
-	TMap<UMaterial *, TArray<FMaterialParameterInfo>> ParametersToVtIze;
-	TMap<UMaterialFunctionInterface *, TArray<FMaterialParameterInfo>> FunctionParametersToVtIze;
+	TMap<UMaterial*, TSet<FMaterialParameterInfo>> ParametersToVtIze;
+	TMap<UMaterialFunctionInterface*, TSet<FMaterialParameterInfo>> FunctionParametersToVtIze;
 
 	// Find all materials that reference the passed in textures
 	// This will also load these materials.
 	TArray<UMaterialInterface *> MaterialsUsedByAffectedTextures;
-	int TextureIndex = 0;
-	while (TextureIndex < InAffectedTextures.Num())
+	for (TSet<UTexture2D*>::TIterator TextureIter(InAffectedTextures); TextureIter; ++TextureIter)
 	{
-		UTexture2D* Tex2D = InAffectedTextures[TextureIndex];
+		UTexture2D* Tex2D = *TextureIter;
 
 		Task.EnterProgressFrame();
 		TArray<UMaterialInterface *> MaterialsUsingTexture;
@@ -166,18 +165,17 @@ void FVirtualTextureConversionWorker::FindAllTexturesAndMaterials_Iteration(TArr
 
 				MaterialsUsedByAffectedTextures.Append(MaterialsUsingTexture);
 			}
-			++TextureIndex;
 		}
 		else
 		{
-			InAffectedTextures.RemoveAtSwap(TextureIndex);
+			TextureIter.RemoveCurrent();
 			InInvalidTextures.Add(Tex2D);
 		}
 	}
 
 	// Find all materials that reference the passed in functions
 	// This will also load these materials.
-	for (auto Func : InAffectedFunctions)
+	for (UMaterialFunctionInterface* Func : InAffectedFunctions)
 	{
 		Task.EnterProgressFrame();
 		TArray<UMaterialInterface *> MaterialsUsingFunction;
@@ -211,12 +209,15 @@ void FVirtualTextureConversionWorker::FindAllTexturesAndMaterials_Iteration(TArr
 		check(Material); // It's something else entirely??
 		MaterialInferfaces.AddUnique(Material);
 		MaterialHeap.AddUnique(Material);
-		InAffectedMaterials.AddUnique(Material);
-
-		AuditTrail.Add(Material, FAuditTrail(
-			If,
-			FString::Printf(TEXT("is the child  of"))
-		));
+		bool bAlreadyExists;
+		InAffectedMaterials.Add(Material, &bAlreadyExists);
+		if (!bAlreadyExists)
+		{
+			AuditTrail.Add(Material, FAuditTrail(
+				If,
+				FString::Printf(TEXT("is the child  of"))
+			));
+		}
 	}
 
 	// We now have a set of "root" materials which will be affected by changing InTextures to VT.
@@ -243,7 +244,7 @@ void FVirtualTextureConversionWorker::FindAllTexturesAndMaterials_Iteration(TArr
 		TArray<FGuid> ParameterGuids;
 		ParentMaterial->GetAllTextureParameterInfo(ParameterInfos, ParameterGuids);
 
-		for (auto ParamInfo : ParameterInfos)
+		for (FMaterialParameterInfo& ParamInfo : ParameterInfos)
 		{
 			UTexture *ParamValue = nullptr;
 			if (ParentMaterial->GetTextureParameterValue(ParamInfo, ParamValue))
@@ -292,11 +293,15 @@ void FVirtualTextureConversionWorker::FindAllTexturesAndMaterials_Iteration(TArr
 				UTexture2D *Tex2D = Cast<UTexture2D>(Tex);
 				if (Tex2D && !Tex2D->VirtualTextureStreaming)
 				{
-					InAffectedTextures.AddUnique(Tex2D);
-					AuditTrail.Add(Tex2D, FAuditTrail(
-						Mat,
-						FString::Printf(TEXT("set on parameter %s in instance %s of material"), *Parameter.Name.ToString(), *If->GetName())
-					));
+					bool bAlreadyExists;
+					InAffectedTextures.Add(Tex2D, &bAlreadyExists);
+					if (!bAlreadyExists)
+					{
+						AuditTrail.Add(Tex2D, FAuditTrail(
+							Mat,
+							FString::Printf(TEXT("set on parameter %s in instance %s of material"), *Parameter.Name.ToString(), *If->GetName())
+						));
+					}
 				}
 			}
 		}
@@ -308,7 +313,7 @@ void FVirtualTextureConversionWorker::FindAllTexturesAndMaterials_Iteration(TArr
 
 	// Find all materials functions that directly reference the passed in textures
 	TArray<UMaterialFunctionInterface *> FunctionsUsedByAffectedTextures;
-	for (auto Tex2D : InAffectedTextures)
+	for (UTexture2D* Tex2D : InAffectedTextures)
 	{
 		Task.EnterProgressFrame();
 		if (!Tex2D->GetPathName().StartsWith("/Engine/"))
@@ -327,25 +332,28 @@ void FVirtualTextureConversionWorker::FindAllTexturesAndMaterials_Iteration(TArr
 		{
 			FunctionInferfaces.AddUnique(Function);
 			FunctionHeap.AddUnique(Function);
-			InAffectedFunctions.AddUnique(Function);
-
-			AuditTrail.Add(Function, FAuditTrail(
-				If,
-				FString::Printf(TEXT("this is the parent of"))
-			));
+			bool bAlreadyExists;
+			InAffectedFunctions.Add(Function, &bAlreadyExists);
+			if (!bAlreadyExists)
+			{
+				AuditTrail.Add(Function, FAuditTrail(
+					If,
+					FString::Printf(TEXT("this is the parent of"))
+				));
+			}
 		}
 	}
 
 
 	// We have a second class of functions here. That is functions that don't directly reference any textures of interest
 	// but where there are material instances that overrides properties in the textures that do reference textures of interest
-	for (auto MaterialParametersPair : ParametersToVtIze)
+	for (TPair<UMaterial*, TSet<FMaterialParameterInfo>>& MaterialParametersPair : ParametersToVtIze)
 	{
-		for (auto Parameter : MaterialParametersPair.Value)
+		for (FMaterialParameterInfo& Parameter : MaterialParametersPair.Value)
 		{
 			TArray<UMaterialFunctionInterface*> DependentFunctions;
 			MaterialParametersPair.Key->GetDependentFunctions(DependentFunctions);
-			for (auto Function : DependentFunctions)
+			for (UMaterialFunctionInterface* Function : DependentFunctions)
 			{
 				for (UMaterialExpression* FunctionExpression : Function->GetExpressions())
 				{
@@ -353,8 +361,8 @@ void FVirtualTextureConversionWorker::FindAllTexturesAndMaterials_Iteration(TArr
 					{
 						if (TexParameter->ParameterName == Parameter.Name)
 						{
-							FunctionParametersToVtIze.FindOrAdd(Function->GetBaseFunction()).AddUnique(Parameter);
-							InAffectedFunctions.AddUnique(Function->GetBaseFunction());
+							FunctionParametersToVtIze.FindOrAdd(Function->GetBaseFunction()).Add(Parameter);
+							InAffectedFunctions.Add(Function->GetBaseFunction());
 							FunctionHeap.AddUnique(Function->GetBaseFunction());
 						}
 					}
@@ -417,11 +425,15 @@ void FVirtualTextureConversionWorker::FindAllTexturesAndMaterials_Iteration(TArr
 			UTexture2D *Tex2D = Cast<UTexture2D>(Tex);
 			if (Tex2D && !Tex2D->VirtualTextureStreaming)
 			{
-				InAffectedTextures.AddUnique(Tex2D);
-				AuditTrail.Add(Tex2D, FAuditTrail(
-					Func,
-					FString::Printf(TEXT("set on parameter %s in instance %s"), *Parameter.Name.ToString(), *If->GetPathName())
-				));
+				bool bAlreadyExists;
+				InAffectedTextures.Add(Tex2D, &bAlreadyExists);
+				if (!bAlreadyExists)
+				{
+					AuditTrail.Add(Tex2D, FAuditTrail(
+						Func,
+						FString::Printf(TEXT("set on parameter %s in instance %s"), *Parameter.Name.ToString(), *If->GetPathName())
+					));
+				}
 			}
 		}
 	}
@@ -429,24 +441,32 @@ void FVirtualTextureConversionWorker::FindAllTexturesAndMaterials_Iteration(TArr
 
 void FVirtualTextureConversionWorker::FindAllTexturesAndMaterials(TArray<UMaterial *> &OutAffectedMaterials, TArray<UMaterialFunctionInterface *> &OutAffectedFunctions, TArray<UTexture2D *> &OutAffectedTextures)
 {
-	int LastNumMaterials = OutAffectedMaterials.Num();
-	int LastNumTextures = OutAffectedTextures.Num();
-	int LastNumFunctions = OutAffectedFunctions.Num();
-
 	FScopedSlowTask SlowTask(1000.0f, LOCTEXT("ConvertToVT_Progress_FindAllTexturesAndMaterials", "Finding Textures and Materials..."));
+	TSet<UMaterial*> AffectedMaterials;
+	TSet<UMaterialFunctionInterface*> AffectedFunctions;
+	TSet<UTexture2D*> AffectedTextures;
+	AffectedMaterials.Append(OutAffectedMaterials);
+	AffectedFunctions.Append(OutAffectedFunctions);
+	AffectedTextures.Append(OutAffectedTextures);
+	int LastNumMaterials = AffectedMaterials.Num();
+	int LastNumTextures = AffectedTextures.Num();
+	int LastNumFunctions = AffectedFunctions.Num();
 
 	while (true)
 	{
-		FindAllTexturesAndMaterials_Iteration(OutAffectedMaterials, OutAffectedFunctions, OutAffectedTextures, MaterialRejectedTextures, SlowTask);
-		if (OutAffectedMaterials.Num() == LastNumMaterials && OutAffectedTextures.Num() == LastNumTextures && OutAffectedFunctions.Num() == LastNumFunctions)
+		FindAllTexturesAndMaterials_Iteration(AffectedMaterials, AffectedFunctions, AffectedTextures, MaterialRejectedTextures, SlowTask);
+		if (AffectedMaterials.Num() == LastNumMaterials && AffectedTextures.Num() == LastNumTextures && AffectedFunctions.Num() == LastNumFunctions)
 		{
-			return;
+			break;
 		}
 
 		LastNumMaterials = OutAffectedMaterials.Num();
 		LastNumTextures = OutAffectedTextures.Num();
 		LastNumFunctions = OutAffectedFunctions.Num();
 	}
+	OutAffectedMaterials = AffectedMaterials.Array();
+	OutAffectedFunctions = AffectedFunctions.Array();
+	OutAffectedTextures = AffectedTextures.Array();
 }
 
 void FVirtualTextureConversionWorker::FilterList(int32 SizeThreshold)
diff --git a/Engine/Source/Editor/VirtualTexturingEditor/Private/VirtualTextureConversionWorker.h b/Engine/Source/Editor/VirtualTexturingEditor/Private/VirtualTextureConversionWorker.h
--- a/Engine/Source/Editor/VirtualTexturingEditor/Private/VirtualTextureConversionWorker.h
+++ b/Engine/Source/Editor/VirtualTexturingEditor/Private/VirtualTextureConversionWorker.h
@@ -29,7 +29,7 @@ public:
 	TArray<UMaterial *>Materials;
 	TArray<UMaterialFunctionInterface *> Functions;
 	TArray<UTexture2D*> SizeRejectedTextures;// The original selection of the user filtered by textures that match the threshold size
-	TArray<UTexture2D*> MaterialRejectedTextures; // Textures that could not be converted due to their usage in materials (connected to unsupported property)
+	TSet<UTexture2D*> MaterialRejectedTextures; // Textures that could not be converted due to their usage in materials (connected to unsupported property)
 
 	struct FAuditTrail
 	{
@@ -54,10 +54,10 @@ public:
 	}
 
 private:
-	void FindAllTexturesAndMaterials_Iteration(TArray<UMaterial*>& InAffectedMaterials,
-		TArray<UMaterialFunctionInterface*>& InAffectedFunctions,
-		TArray<UTexture2D*>& InAffectedTextures,
-		TArray<UTexture2D*>& InInvalidTextures,
+	void FindAllTexturesAndMaterials_Iteration(TSet<UMaterial*>& InAffectedMaterials,
+		TSet<UMaterialFunctionInterface*>& InAffectedFunctions,
+		TSet<UTexture2D*>& InAffectedTextures,
+		TSet<UTexture2D*>& InInvalidTextures,
 		FScopedSlowTask& Task);
 
 	void FindAllTexturesAndMaterials(TArray<UMaterial *> &OutAffectedMaterials, TArray<UMaterialFunctionInterface *> &OutAffectedFunctions, TArray<UTexture2D *> &OutAffectedTextures);
