Convert To VT : only UpdateList once, after user textures are set

diff --git a/Engine/Source/Editor/VirtualTexturingEditor/Private/SConvertToVirtualTexture.cpp b/Engine/Source/Editor/VirtualTexturingEditor/Private/SConvertToVirtualTexture.cpp
--- a/Engine/Source/Editor/VirtualTexturingEditor/Private/SConvertToVirtualTexture.cpp
+++ b/Engine/Source/Editor/VirtualTexturingEditor/Private/SConvertToVirtualTexture.cpp
@@ -178,7 +178,8 @@ void SConvertToVirtualTexture::Construct(const FArguments& InArgs)
 			]
 	];
 
-	UpdateList();
+	// will be done by SetUserTextures :
+	//UpdateList();
 	bFilterButtonEnabled = false;
 }
 
@@ -615,6 +616,7 @@ FConvertToVTDlg::FConvertToVTDlg(const TArray<UTexture2D *> &Textures, bool bBac
 				.ParentWindow(DialogWindow)
 			];
 
+		// SetUserTextures will do UpdateList :
 		DialogWidget->SetBackwards(bBackwards);
 		DialogWidget->SetUserTextures(Textures);
 		DialogWindow->SetContent(DialogWrapper.ToSharedRef());
