fix infinite loop introduced in Convert To VT

diff --git a/Engine/Source/Editor/VirtualTexturingEditor/Private/VirtualTextureConversionWorker.cpp b/Engine/Source/Editor/VirtualTexturingEditor/Private/VirtualTextureConversionWorker.cpp
--- a/Engine/Source/Editor/VirtualTexturingEditor/Private/VirtualTextureConversionWorker.cpp
+++ b/Engine/Source/Editor/VirtualTexturingEditor/Private/VirtualTextureConversionWorker.cpp
@@ -460,9 +460,9 @@ void FVirtualTextureConversionWorker::FindAllTexturesAndMaterials(TArray<UMateri
 			break;
 		}
 
-		LastNumMaterials = OutAffectedMaterials.Num();
-		LastNumTextures = OutAffectedTextures.Num();
-		LastNumFunctions = OutAffectedFunctions.Num();
+		LastNumMaterials = AffectedMaterials.Num();
+		LastNumTextures = AffectedTextures.Num();
+		LastNumFunctions = AffectedFunctions.Num();
 	}
 	OutAffectedMaterials = AffectedMaterials.Array();
 	OutAffectedFunctions = AffectedFunctions.Array();
