Added missing FSyncScope for VT thumbnail renderer.

diff --git a/Engine/Source/Editor/VirtualTexturingEditor/Private/RuntimeVirtualTextureThumbnailRenderer.cpp b/Engine/Source/Editor/VirtualTexturingEditor/Private/RuntimeVirtualTextureThumbnailRenderer.cpp
--- a/Engine/Source/Editor/VirtualTexturingEditor/Private/RuntimeVirtualTextureThumbnailRenderer.cpp
+++ b/Engine/Source/Editor/VirtualTexturingEditor/Private/RuntimeVirtualTextureThumbnailRenderer.cpp
@@ -95,6 +95,8 @@ void URuntimeVirtualTextureThumbnailRenderer::Draw(UObject* Object, int32 X, int
 	RuntimeVirtualTexture->GetProducerDescription(VTDesc, URuntimeVirtualTexture::FInitSettings(), Transform);
 	const int32 MaxLevel = (int32)FMath::CeilLogTwo(FMath::Max(VTDesc.BlockWidthInTiles, VTDesc.BlockHeightInTiles));
 
+	UE::RenderCommandPipe::FSyncScope SyncScope;
+
 	ENQUEUE_RENDER_COMMAND(BakeStreamingTextureTileCommand)(
 		[Scene, VirtualTextureSceneIndex, MaterialType, RenderTarget, DestBox, Transform, Bounds, MaxLevel](FRHICommandListImmediate& RHICmdList)
 	{
