When calculating RVT bounds, if the bounds align actor is a landscape, then include all of its landscape components.

diff --git a/Engine/Source/Editor/VirtualTexturingEditor/Private/RuntimeVirtualTextureSetBounds.cpp b/Engine/Source/Editor/VirtualTexturingEditor/Private/RuntimeVirtualTextureSetBounds.cpp
--- a/Engine/Source/Editor/VirtualTexturingEditor/Private/RuntimeVirtualTextureSetBounds.cpp
+++ b/Engine/Source/Editor/VirtualTexturingEditor/Private/RuntimeVirtualTextureSetBounds.cpp
@@ -29,12 +29,30 @@ namespace RuntimeVirtualTexture
 		LocalTransform.SetComponents(TargetRotation, TargetPosition, FVector::OneVector);
 		FTransform WorldToLocal = LocalTransform.Inverse();
 
+		// Special case where if the bounds align actor is a landscape we want to automatically include all associated landscape components.
+		FGuid BoundsAlignLandscapeGuid;
+		if (BoundsAlignActor.IsValid())
+		{
+			if (ALandscape const* LandscapeProxy = Cast<ALandscape>(BoundsAlignActor.Get()))
+			{
+				BoundsAlignLandscapeGuid = LandscapeProxy->GetLandscapeGuid();
+			}
+		}
+
 		// Expand bounds for the BoundsAlignActor and all primitive components that write to this virtual texture.
 		FBox Bounds(ForceInit);
 		for (TObjectIterator<UPrimitiveComponent> It(RF_ClassDefaultObject, true, EInternalObjectFlags::Garbage); It; ++It)
 		{
 			bool bUseBounds = BoundsAlignActor.IsValid() && It->GetOwner() == BoundsAlignActor.Get();
 
+			if (BoundsAlignLandscapeGuid.IsValid())
+			{
+				if (ALandscapeProxy const* LandscapeProxy = Cast<ALandscapeProxy>(It->GetOwner()))
+				{
+					bUseBounds |= LandscapeProxy->GetLandscapeGuid() == BoundsAlignLandscapeGuid;
+				}
+			}
+
 			TArray<URuntimeVirtualTexture*> const& VirtualTextures = It->GetRuntimeVirtualTextures();
 			for (int32 Index = 0; !bUseBounds && Index < VirtualTextures.Num(); ++Index) 
 			{
