Set RVT texture bounds usability improvements :
* Ensure the set bounds are valid in the case of a flat mesh
* Take into account the unloaded components bounds when the bounds align actor is a landscape

diff --git a/Engine/Source/Editor/VirtualTexturingEditor/Private/RuntimeVirtualTextureSetBounds.cpp b/Engine/Source/Editor/VirtualTexturingEditor/Private/RuntimeVirtualTextureSetBounds.cpp
--- a/Engine/Source/Editor/VirtualTexturingEditor/Private/RuntimeVirtualTextureSetBounds.cpp
+++ b/Engine/Source/Editor/VirtualTexturingEditor/Private/RuntimeVirtualTextureSetBounds.cpp
@@ -29,30 +29,24 @@ namespace RuntimeVirtualTexture
 		LocalTransform.SetComponents(TargetRotation, TargetPosition, FVector::OneVector);
 		FTransform WorldToLocal = LocalTransform.Inverse();
 
-		// Special case where if the bounds align actor is a landscape we want to automatically include all associated landscape components.
-		FGuid BoundsAlignLandscapeGuid;
+		// Expand bounds for the BoundsAlignActor and all primitive components that write to this virtual texture.
+		FBox Bounds(ForceInit);
+
+		// Special case where if the bounds align actor is a landscape: we want to automatically include all associated landscape components, 
+		//  including those that are not currently loaded. Luckily there's a function for that:
 		if (BoundsAlignActor.IsValid())
 		{
-			if (ALandscape const* LandscapeProxy = Cast<ALandscape>(BoundsAlignActor.Get()))
+			if (ALandscape const* Landscape = Cast<ALandscape>(BoundsAlignActor.Get()))
 			{
-				BoundsAlignLandscapeGuid = LandscapeProxy->GetLandscapeGuid();
+				FBox WorldBounds = Landscape->GetCompleteBounds();
+				Bounds = WorldBounds.TransformBy(WorldToLocal);
 			}
 		}
 
-		// Expand bounds for the BoundsAlignActor and all primitive components that write to this virtual texture.
-		FBox Bounds(ForceInit);
 		for (TObjectIterator<UPrimitiveComponent> It(RF_ClassDefaultObject, true, EInternalObjectFlags::Garbage); It; ++It)
 		{
 			bool bUseBounds = BoundsAlignActor.IsValid() && It->GetOwner() == BoundsAlignActor.Get();
 
-			if (BoundsAlignLandscapeGuid.IsValid())
-			{
-				if (ALandscapeProxy const* LandscapeProxy = Cast<ALandscapeProxy>(It->GetOwner()))
-				{
-					bUseBounds |= LandscapeProxy->GetLandscapeGuid() == BoundsAlignLandscapeGuid;
-				}
-			}
-
 			TArray<URuntimeVirtualTexture*> const& VirtualTextures = It->GetRuntimeVirtualTextures();
 			for (int32 Index = 0; !bUseBounds && Index < VirtualTextures.Num(); ++Index) 
 			{
@@ -65,18 +59,25 @@ namespace RuntimeVirtualTexture
 			if (bUseBounds)
 			{
 				FBoxSphereBounds LocalSpaceBounds = It->CalcBounds(It->GetComponentTransform() * WorldToLocal);
-				if (LocalSpaceBounds.GetBox().GetVolume() > 0.f)
-				{
-					Bounds += LocalSpaceBounds.GetBox();
-				}
+				Bounds += LocalSpaceBounds.GetBox();
 			}
 		}
 
-		// Expand bounds.
-		const float ExpandBounds = InComponent->GetExpandBounds();
-		if (Bounds.IsValid && ExpandBounds > 0)
+		if (Bounds.IsValid)
 		{
-			Bounds = Bounds.ExpandBy(ExpandBounds);
+			const FVector BoundsSize = Bounds.GetSize();
+			// If XY bounds are valid but Z is 0, let's just expand by a little value to get something to render still (e.g. flat landscape)
+			if ((BoundsSize.X > UE_KINDA_SMALL_NUMBER) && (BoundsSize.Y > UE_KINDA_SMALL_NUMBER) && (BoundsSize.Z <= UE_KINDA_SMALL_NUMBER))
+			{
+				Bounds = Bounds.ExpandBy(FVector(0.0, 0.0, 0.5));
+			}
+
+			// Expand bounds if requested
+			const float ExpandBounds = InComponent->GetExpandBounds();
+			if (ExpandBounds > 0.0f)
+			{
+				Bounds = Bounds.ExpandBy(ExpandBounds);
+			}
 		}
 
 		// Calculate the transform to fit the bounds.
