VirtualTexturingEditor Module Truncation fixes

diff --git a/Engine/Source/Editor/VirtualTexturingEditor/Private/RuntimeVirtualTextureAssetTypeActions.cpp b/Engine/Source/Editor/VirtualTexturingEditor/Private/RuntimeVirtualTextureAssetTypeActions.cpp
--- a/Engine/Source/Editor/VirtualTexturingEditor/Private/RuntimeVirtualTextureAssetTypeActions.cpp
+++ b/Engine/Source/Editor/VirtualTexturingEditor/Private/RuntimeVirtualTextureAssetTypeActions.cpp
@@ -78,7 +78,7 @@ namespace
 			FindAllMaterials(RuntimeVirtualTexture, Materials, Functions);
 
 			int32 TaskCount = Materials.Num() + Functions.Num();
-			FScopedSlowTask Task(TaskCount, LOCTEXT("RuntimeVirtualTexture_FixMaterialUsageProgress", "Fixing materials for Runtime Virtual Texture usage..."));
+			FScopedSlowTask Task(static_cast<float>(TaskCount), LOCTEXT("RuntimeVirtualTexture_FixMaterialUsageProgress", "Fixing materials for Runtime Virtual Texture usage..."));
 			Task.MakeDialog();
 
 			for (UMaterial* Material : Materials)
diff --git a/Engine/Source/Editor/VirtualTexturingEditor/Private/RuntimeVirtualTextureBuildStreamingMips.cpp b/Engine/Source/Editor/VirtualTexturingEditor/Private/RuntimeVirtualTextureBuildStreamingMips.cpp
--- a/Engine/Source/Editor/VirtualTexturingEditor/Private/RuntimeVirtualTextureBuildStreamingMips.cpp
+++ b/Engine/Source/Editor/VirtualTexturingEditor/Private/RuntimeVirtualTextureBuildStreamingMips.cpp
@@ -195,7 +195,7 @@ namespace RuntimeVirtualTexture
 		}
 
 		// Spin up slow task UI
-		const float TaskWorkRender = NumTilesX * NumTilesY;
+		const float TaskWorkRender = static_cast<float>(NumTilesX * NumTilesY);
 		const float TextureBuildTaskMultiplier = 0.25f;
 		const float TaskWorkBuildBulkData = TaskWorkRender * TextureBuildTaskMultiplier;
 		FScopedSlowTask Task(TaskWorkRender + TaskWorkBuildBulkData, FText::AsCultureInvariant(InComponent->GetStreamingTexture()->GetName()));
@@ -263,7 +263,7 @@ namespace RuntimeVirtualTexture
 						Desc.UVToWorld = Transform;
 						Desc.WorldBounds = Bounds;
 						Desc.MaterialType = MaterialType;
-						Desc.MaxLevel = MaxLevel;
+						Desc.MaxLevel = IntCastChecked<uint8>(MaxLevel);
 						Desc.bClearTextures = true;
 						Desc.bIsThumbnails = false;
 						Desc.FixedColor = FixedColor;
@@ -275,7 +275,7 @@ namespace RuntimeVirtualTexture
 						Desc.PageDescs[0].DestBox[1] = TileBox;
 						Desc.PageDescs[0].DestBox[2] = TileBox;
 						Desc.PageDescs[0].UVRange = UVRange;
-						Desc.PageDescs[0].vLevel = RenderLevel;
+						Desc.PageDescs[0].vLevel = IntCastChecked<uint8>(RenderLevel);
 
 						RuntimeVirtualTexture::RenderPagesStandAlone(GraphBuilder, Desc);
 
diff --git a/Engine/Source/Editor/VirtualTexturingEditor/Private/RuntimeVirtualTextureDetailsCustomization.cpp b/Engine/Source/Editor/VirtualTexturingEditor/Private/RuntimeVirtualTextureDetailsCustomization.cpp
--- a/Engine/Source/Editor/VirtualTexturingEditor/Private/RuntimeVirtualTextureDetailsCustomization.cpp
+++ b/Engine/Source/Editor/VirtualTexturingEditor/Private/RuntimeVirtualTextureDetailsCustomization.cpp
@@ -235,7 +235,7 @@ void FRuntimeVirtualTextureComponentDetailsCustomization::CustomizeDetails(IDeta
 		SNew(SButton)
 		.VAlign(VAlign_Center)
 		.HAlign(HAlign_Center)
-		.ContentPadding(2)
+		.ContentPadding(2.f)
 		.Text(LOCTEXT("Button_SetBounds", "Set Bounds"))
 		.OnClicked(this, &FRuntimeVirtualTextureComponentDetailsCustomization::SetBounds)
 		.IsEnabled(this, &FRuntimeVirtualTextureComponentDetailsCustomization::IsSetBoundsEnabled)
@@ -264,7 +264,7 @@ void FRuntimeVirtualTextureComponentDetailsCustomization::CustomizeDetails(IDeta
 			SNew(SButton)
 			.HAlign(HAlign_Center)
 			.VAlign(VAlign_Center)
-			.ContentPadding(2)
+			.ContentPadding(2.f)
 			.Text(LOCTEXT("Button_Build", "Build"))
 			.ToolTipText(BuildButtonText)
 			.OnClicked(this, &FRuntimeVirtualTextureComponentDetailsCustomization::BuildStreamedMips)
diff --git a/Engine/Source/Editor/VirtualTexturingEditor/Private/RuntimeVirtualTextureSetBounds.cpp b/Engine/Source/Editor/VirtualTexturingEditor/Private/RuntimeVirtualTextureSetBounds.cpp
--- a/Engine/Source/Editor/VirtualTexturingEditor/Private/RuntimeVirtualTextureSetBounds.cpp
+++ b/Engine/Source/Editor/VirtualTexturingEditor/Private/RuntimeVirtualTextureSetBounds.cpp
@@ -119,8 +119,8 @@ namespace RuntimeVirtualTexture
 				// Adjust position to snap at a half texel offset from landscape.
 				const FVector BaseVirtualTexturePosition = Transform.GetTranslation();
 				const FVector LandscapeSnapPosition = LandscapePosition - 0.5f * VirtualTexelWorldSize;
-				const float SnapOffsetX = FMath::Frac((BaseVirtualTexturePosition.X - LandscapeSnapPosition.X) / VirtualTexelWorldSize.X) * VirtualTexelWorldSize.X;
-				const float SnapOffsetY = FMath::Frac((BaseVirtualTexturePosition.Y - LandscapeSnapPosition.Y) / VirtualTexelWorldSize.Y) * VirtualTexelWorldSize.Y;
+				const double SnapOffsetX = FMath::Frac((BaseVirtualTexturePosition.X - LandscapeSnapPosition.X) / VirtualTexelWorldSize.X) * VirtualTexelWorldSize.X;
+				const double SnapOffsetY = FMath::Frac((BaseVirtualTexturePosition.Y - LandscapeSnapPosition.Y) / VirtualTexelWorldSize.Y) * VirtualTexelWorldSize.Y;
 				const FVector VirtualTexturePosition = BaseVirtualTexturePosition - FVector(SnapOffsetX, SnapOffsetY, 0);
 				Transform.SetTranslation(FVector(BaseVirtualTexturePosition.X - SnapOffsetX, BaseVirtualTexturePosition.Y - SnapOffsetY, BaseVirtualTexturePosition.Z));
 			}
diff --git a/Engine/Source/Editor/VirtualTexturingEditor/Private/RuntimeVirtualTextureThumbnailRenderer.cpp b/Engine/Source/Editor/VirtualTexturingEditor/Private/RuntimeVirtualTextureThumbnailRenderer.cpp
--- a/Engine/Source/Editor/VirtualTexturingEditor/Private/RuntimeVirtualTextureThumbnailRenderer.cpp
+++ b/Engine/Source/Editor/VirtualTexturingEditor/Private/RuntimeVirtualTextureThumbnailRenderer.cpp
@@ -110,7 +110,7 @@ void URuntimeVirtualTextureThumbnailRenderer::Draw(UObject* Object, int32 X, int
 		Desc.UVToWorld = Transform;
 		Desc.WorldBounds = Bounds;
 		Desc.MaterialType = MaterialType;
-		Desc.MaxLevel = MaxLevel;
+		Desc.MaxLevel = IntCastChecked<uint8>(MaxLevel);
 		Desc.bClearTextures = true;
 		Desc.bIsThumbnails = true;
 		Desc.FixedColor = FLinearColor::Transparent;
@@ -118,7 +118,7 @@ void URuntimeVirtualTextureThumbnailRenderer::Draw(UObject* Object, int32 X, int
 		Desc.Targets[0].Texture = RenderTarget->GetRenderTargetTexture();
 		Desc.PageDescs[0].DestBox[0] = DestBox;
 		Desc.PageDescs[0].UVRange = FBox2D(FVector2D(0, 0), FVector2D(1, 1));
-		Desc.PageDescs[0].vLevel = MaxLevel;
+		Desc.PageDescs[0].vLevel = IntCastChecked<uint8>(MaxLevel);
 
 		RuntimeVirtualTexture::RenderPagesStandAlone(GraphBuilder, Desc);
 
diff --git a/Engine/Source/Editor/VirtualTexturingEditor/Private/VirtualTextureConversionWorker.cpp b/Engine/Source/Editor/VirtualTexturingEditor/Private/VirtualTextureConversionWorker.cpp
--- a/Engine/Source/Editor/VirtualTexturingEditor/Private/VirtualTextureConversionWorker.cpp
+++ b/Engine/Source/Editor/VirtualTexturingEditor/Private/VirtualTextureConversionWorker.cpp
@@ -37,7 +37,7 @@ namespace VTConversionWorkerUtil
 		TArray<FAssetData> OutAssetDatas;
 		UAssetRegistryHelpers::FindReferencersOfAssetOfClass(Object, { T::StaticClass() }, OutAssetDatas);
 
-		FScopedSlowTask SlowTask(OutAssetDatas.Num(), LOCTEXT("ConvertToVT_Progress_LoadingObjects", "Loading Objects..."));
+		FScopedSlowTask SlowTask(static_cast<float>(OutAssetDatas.Num()), LOCTEXT("ConvertToVT_Progress_LoadingObjects", "Loading Objects..."));
 
 		for (auto Data : OutAssetDatas)
 		{
@@ -542,7 +542,7 @@ void FVirtualTextureConversionWorker::DoConvert()
 	UE_LOG(LogVirtualTextureConversion, Display, TEXT("Beginning conversion..."));
 	SlowTask.EnterProgressFrame();
 	{
-		FScopedSlowTask TextureTask(Textures.Num(), LOCTEXT("ConvertToVT_Progress_TextureTask", "Updating textures..."));
+		FScopedSlowTask TextureTask(static_cast<float>(Textures.Num()), LOCTEXT("ConvertToVT_Progress_TextureTask", "Updating textures..."));
 
 		for (UTexture2D *Tex : Textures)
 		{
@@ -605,7 +605,7 @@ void FVirtualTextureConversionWorker::DoConvert()
 
 	SlowTask.EnterProgressFrame();
 	{
-		FScopedSlowTask MaterialTask(Materials.Num() + Functions.Num(), LOCTEXT("ConvertToVT_Progress_MaterialTask", "Updating materials..."));
+		FScopedSlowTask MaterialTask(static_cast<float>(Materials.Num() + Functions.Num()), LOCTEXT("ConvertToVT_Progress_MaterialTask", "Updating materials..."));
 		FMaterialUpdateContext UpdateContext;
 
 		TMap<UMaterialFunctionInterface*, TArray<UMaterial*>> FunctionToMaterialMap;
diff --git a/Engine/Source/Editor/VirtualTexturingEditor/VirtualTexturingEditor.Build.cs b/Engine/Source/Editor/VirtualTexturingEditor/VirtualTexturingEditor.Build.cs
--- a/Engine/Source/Editor/VirtualTexturingEditor/VirtualTexturingEditor.Build.cs
+++ b/Engine/Source/Editor/VirtualTexturingEditor/VirtualTexturingEditor.Build.cs
@@ -6,8 +6,8 @@ public class VirtualTexturingEditor : ModuleRules
 {
 	public VirtualTexturingEditor(ReadOnlyTargetRules Target) : base(Target)
 	{
-        PrivateDependencyModuleNames.AddRange(
-            new string[] {
+		PrivateDependencyModuleNames.AddRange(
+			new string[] {
 				"AppFramework",
 				"AssetRegistry",
 				"ContentBrowser",
@@ -26,7 +26,9 @@ public class VirtualTexturingEditor : ModuleRules
 				"Slate",
 				"SlateCore",
 				"UnrealEd",
-            }
-        );
-    }
+			}
+		);
+
+		UnsafeTypeCastWarningLevel = WarningLevel.Error;
+	}
 }
\ No newline at end of file
diff --git a/Engine/Source/Runtime/Landscape/Classes/LandscapeComponent.h b/Engine/Source/Runtime/Landscape/Classes/LandscapeComponent.h
--- a/Engine/Source/Runtime/Landscape/Classes/LandscapeComponent.h
+++ b/Engine/Source/Runtime/Landscape/Classes/LandscapeComponent.h
@@ -866,7 +866,7 @@ public:
 	
 	bool MaterialHasGrass() const { return !GetGrassTypes().IsEmpty(); }
 
-	float GetGrassTypesMaxDiscardDistance() const { return GrassTypeSummary.MaxInstanceDiscardDistance; }
+	double GetGrassTypesMaxDiscardDistance() const { return GrassTypeSummary.MaxInstanceDiscardDistance; }
 	void SetGrassTypesMaxDiscardDistance(const float InGrassTypesMaxDiscardDistance) { GrassTypeSummary.MaxInstanceDiscardDistance = InGrassTypesMaxDiscardDistance; GrassTypeSummary.bInvalid = false; }
 
 	/** If the LandscapeMaterial has changed, updates the GrassTypes array. Returns true if the GrassTypes array was updated. */
diff --git a/Engine/Source/Runtime/Landscape/Public/LandscapeDataAccess.h b/Engine/Source/Runtime/Landscape/Public/LandscapeDataAccess.h
--- a/Engine/Source/Runtime/Landscape/Public/LandscapeDataAccess.h
+++ b/Engine/Source/Runtime/Landscape/Public/LandscapeDataAccess.h
@@ -45,7 +45,7 @@ namespace LandscapeDataAccess
 
 	FORCEINLINE float UnpackHeight(const FColor& InHeightmapSample)
 	{
-		uint16 Height = (InHeightmapSample.R << 8) + InHeightmapSample.G;
+		uint16 Height = (uint16)(InHeightmapSample.R << 8) | (uint16)InHeightmapSample.G;
 		return GetLocalHeight(Height);
 	}
 
@@ -211,7 +211,7 @@ struct FLandscapeComponentDataInterfaceBase
 	uint16 GetHeight(int32 LocalX, int32 LocalY, const TArray<FColor>& HeightAndNormals) const
 	{
 		const FColor* Texel = GetHeightData(LocalX, LocalY, HeightAndNormals);
-		return (Texel->R << 8) + Texel->G;
+		return (uint16)(Texel->R << 8) | (uint16)Texel->G;
 	}
 
 	float GetScaleFactor() const
@@ -223,7 +223,7 @@ struct FLandscapeComponentDataInterfaceBase
 	{
 		const float ScaleFactor = GetScaleFactor();
 		
-		return FVector(LocalX * ScaleFactor , LocalY * ScaleFactor, LandscapeDataAccess::GetLocalHeight(GetHeight(LocalX, LocalY, HeightAndNormals)));
+		return FVector(ScaleFactor * (float)LocalX, ScaleFactor * (float)LocalY, LandscapeDataAccess::GetLocalHeight(GetHeight(LocalX, LocalY, HeightAndNormals)));
 	}
 
 	float GetLocalHeight(int32 LocalX, int32 LocalY, const TArray<FColor>& HeightAndNormals) const
@@ -327,7 +327,7 @@ struct FLandscapeComponentDataInterface : public FLandscapeComponentDataInterfac
 	uint16 GetHeight( int32 LocalX, int32 LocalY ) const
 	{
 		FColor* Texel = GetHeightData(LocalX, LocalY);
-		return (Texel->R << 8) + Texel->G;
+		return (uint16)(Texel->R << 8) | (uint16)Texel->G;
 	}
 
 	uint16 GetHeight( int32 VertexIndex ) const
diff --git a/Engine/Source/Runtime/Landscape/Public/LandscapeEditTypes.h b/Engine/Source/Runtime/Landscape/Public/LandscapeEditTypes.h
--- a/Engine/Source/Runtime/Landscape/Public/LandscapeEditTypes.h
+++ b/Engine/Source/Runtime/Landscape/Public/LandscapeEditTypes.h
@@ -39,9 +39,9 @@ ENUM_CLASS_FLAGS(EOutdatedDataFlags);
 
 inline uint32 GetOutdatedDataFlagIndex(EOutdatedDataFlags InFlag)
 {
-	const uint64 InFlagAsUInt64 = static_cast<uint64>(InFlag);
-	check((InFlagAsUInt64 > static_cast<uint64>(EOutdatedDataFlags::None)) && (InFlagAsUInt64 < static_cast<uint64>(EOutdatedDataFlags::LastPlusOne)) && (FMath::CountBits(InFlagAsUInt64) == 1u));
-	return FMath::CountTrailingZeros(InFlagAsUInt64);
+	const uint32 InFlagAsUInt32 = static_cast<uint32>(InFlag);
+	check((InFlagAsUInt32 > static_cast<uint32>(EOutdatedDataFlags::None)) && (InFlagAsUInt32 < static_cast<uint32>(EOutdatedDataFlags::LastPlusOne)) && (FMath::CountBits(InFlagAsUInt32) == 1u));
+	return FMath::CountTrailingZeros(InFlagAsUInt32);
 }
 
 } // namespace UE::Landscape
\ No newline at end of file
diff --git a/Engine/Source/Runtime/RenderCore/Public/RenderGraphEvent.h b/Engine/Source/Runtime/RenderCore/Public/RenderGraphEvent.h
--- a/Engine/Source/Runtime/RenderCore/Public/RenderGraphEvent.h
+++ b/Engine/Source/Runtime/RenderCore/Public/RenderGraphEvent.h
@@ -257,7 +257,7 @@ struct FRDGScope
 		}
 
 		template <typename TScopeType>
-		static constexpr uint32 GetTypeIndex()
+		static constexpr SIZE_T GetTypeIndex()
 		{
 			return TImpl::template IndexOfType<TScopeType>();
 		}
