Fixed crash in RVT details customization when the underlying get GCed

diff --git a/Engine/Source/Editor/VirtualTexturingEditor/Classes/RuntimeVirtualTextureDetailsCustomization.h b/Engine/Source/Editor/VirtualTexturingEditor/Classes/RuntimeVirtualTextureDetailsCustomization.h
--- a/Engine/Source/Editor/VirtualTexturingEditor/Classes/RuntimeVirtualTextureDetailsCustomization.h
+++ b/Engine/Source/Editor/VirtualTexturingEditor/Classes/RuntimeVirtualTextureDetailsCustomization.h
@@ -31,7 +31,7 @@ protected:
 	//~ End IDetailCustomization Interface.
 
 private:
-	URuntimeVirtualTexture* VirtualTexture = nullptr;
+	TWeakObjectPtr<URuntimeVirtualTexture> VirtualTexture;
 
 	IDetailLayoutBuilder* CachedDetailBuilder = nullptr;
 
@@ -70,5 +70,5 @@ protected:
 	//~ End IDetailCustomization Interface.
 
 private:
-	URuntimeVirtualTextureComponent* RuntimeVirtualTextureComponent;
+	TWeakObjectPtr<URuntimeVirtualTextureComponent> RuntimeVirtualTextureComponent;
 };
diff --git a/Engine/Source/Editor/VirtualTexturingEditor/Private/RuntimeVirtualTextureDetailsCustomization.cpp b/Engine/Source/Editor/VirtualTexturingEditor/Private/RuntimeVirtualTextureDetailsCustomization.cpp
--- a/Engine/Source/Editor/VirtualTexturingEditor/Private/RuntimeVirtualTextureDetailsCustomization.cpp
+++ b/Engine/Source/Editor/VirtualTexturingEditor/Private/RuntimeVirtualTextureDetailsCustomization.cpp
@@ -26,7 +26,6 @@
 #define LOCTEXT_NAMESPACE "VirtualTexturingEditorModule"
 
 FRuntimeVirtualTextureDetailsCustomization::FRuntimeVirtualTextureDetailsCustomization()
-	: VirtualTexture(nullptr)
 {
 }
 
@@ -155,35 +154,38 @@ void FRuntimeVirtualTextureDetailsCustomization::CustomizeDetails(IDetailLayoutB
 
 void FRuntimeVirtualTextureDetailsCustomization::RefreshTextDetails()
 {
-	FNumberFormattingOptions SizeOptions;
-	SizeOptions.UseGrouping = false;
-	SizeOptions.MaximumFractionalDigits = 0;
-
- 	TileCountText->SetText(FText::Format(LOCTEXT("Details_Number", "{0}"), FText::AsNumber(VirtualTexture->GetTileCount(), &SizeOptions)));
-	TileSizeText->SetText(FText::Format(LOCTEXT("Details_Number", "{0}"), FText::AsNumber(VirtualTexture->GetTileSize(), &SizeOptions)));
- 	TileBorderSizeText->SetText(FText::Format(LOCTEXT("Details_Number", "{0}"), FText::AsNumber(VirtualTexture->GetTileBorderSize(), &SizeOptions)));
-
-	FString SizeUnits = TEXT("Texels");
-	int32 Size = VirtualTexture->GetSize();
-	int32 SizeLog2 = FMath::CeilLogTwo(Size);
-	if (SizeLog2 >= 30)
+	if (URuntimeVirtualTexture* Texture = VirtualTexture.Get())
 	{
-		Size = Size >> 30;
-		SizeUnits = TEXT("GiTexels");
-	}
-	else if (SizeLog2 >= 20)
-	{
-		Size = Size >> 20;
-		SizeUnits = TEXT("MiTexels");
-	}
-	else if (SizeLog2 >= 10)
-	{
-		Size = Size >> 10;
-		SizeUnits = TEXT("KiTexels");
-	}
-	SizeText->SetText(FText::Format(LOCTEXT("Details_Number_Units", "{0} {1}"), FText::AsNumber(Size, &SizeOptions), FText::FromString(SizeUnits)));
+		FNumberFormattingOptions SizeOptions;
+		SizeOptions.UseGrouping = false;
+		SizeOptions.MaximumFractionalDigits = 0;
+
+		TileCountText->SetText(FText::Format(LOCTEXT("Details_Number", "{0}"), FText::AsNumber(Texture->GetTileCount(), &SizeOptions)));
+		TileSizeText->SetText(FText::Format(LOCTEXT("Details_Number", "{0}"), FText::AsNumber(Texture->GetTileSize(), &SizeOptions)));
+		TileBorderSizeText->SetText(FText::Format(LOCTEXT("Details_Number", "{0}"), FText::AsNumber(Texture->GetTileBorderSize(), &SizeOptions)));
+
+		FString SizeUnits = TEXT("Texels");
+		int32 Size = Texture->GetSize();
+		int32 SizeLog2 = FMath::CeilLogTwo(Size);
+		if (SizeLog2 >= 30)
+		{
+			Size = Size >> 30;
+			SizeUnits = TEXT("GiTexels");
+		}
+		else if (SizeLog2 >= 20)
+		{
+			Size = Size >> 20;
+			SizeUnits = TEXT("MiTexels");
+		}
+		else if (SizeLog2 >= 10)
+		{
+			Size = Size >> 10;
+			SizeUnits = TEXT("KiTexels");
+		}
+		SizeText->SetText(FText::Format(LOCTEXT("Details_Number_Units", "{0} {1}"), FText::AsNumber(Size, &SizeOptions), FText::FromString(SizeUnits)));
 
-	PageTableSizeText->SetText(FText::Format(LOCTEXT("Details_Number", "{0}"), FText::AsNumber(VirtualTexture->GetPageTableSize(), &SizeOptions)));
+		PageTableSizeText->SetText(FText::Format(LOCTEXT("Details_Number", "{0}"), FText::AsNumber(Texture->GetPageTableSize(), &SizeOptions)));
+	}
 }
 
 void FRuntimeVirtualTextureDetailsCustomization::RefreshDetailsView()
@@ -284,71 +286,90 @@ void FRuntimeVirtualTextureComponentDetailsCustomization::CustomizeDetails(IDeta
 
 bool FRuntimeVirtualTextureComponentDetailsCustomization::IsSetBoundsEnabled() const
 {
-	return RuntimeVirtualTextureComponent->GetVirtualTexture() != nullptr;
+	if (URuntimeVirtualTextureComponent* Component = RuntimeVirtualTextureComponent.Get())
+	{
+		return Component->GetVirtualTexture() != nullptr;
+	}
+	return false;
 }
 
 FReply FRuntimeVirtualTextureComponentDetailsCustomization::SetBounds()
 {
-	if (RuntimeVirtualTextureComponent->GetVirtualTexture() != nullptr)
+	if (URuntimeVirtualTextureComponent* Component = RuntimeVirtualTextureComponent.Get())
 	{
-		const FScopedTransaction Transaction(LOCTEXT("Transaction_SetBounds", "Set RuntimeVirtualTextureComponent Bounds"));
-		RuntimeVirtualTexture::SetBounds(RuntimeVirtualTextureComponent);
-		// Force update of editor view widget.
-		GEditor->NoteSelectionChange(false);
-		return FReply::Handled();
+		if (Component->GetVirtualTexture() != nullptr)
+		{
+			const FScopedTransaction Transaction(LOCTEXT("Transaction_SetBounds", "Set RuntimeVirtualTextureComponent Bounds"));
+			RuntimeVirtualTexture::SetBounds(Component);
+			// Force update of editor view widget.
+			GEditor->NoteSelectionChange(false);
+			return FReply::Handled();
+		}
 	}
 	return FReply::Unhandled();
 }
 
 bool FRuntimeVirtualTextureComponentDetailsCustomization::IsBuildStreamedMipsEnabled() const
 {
-	return RuntimeVirtualTextureComponent->GetVirtualTexture() != nullptr;
+	if (URuntimeVirtualTextureComponent* Component = RuntimeVirtualTextureComponent.Get())
+	{
+		return Component->GetVirtualTexture() != nullptr;
+	}
+	return false;
 }
 
 EVisibility FRuntimeVirtualTextureComponentDetailsCustomization::IsBuildWarningIconVisible() const
 {
-	const bool bVisible = RuntimeVirtualTextureComponent->IsStreamingTextureInvalid();
-	return bVisible ? EVisibility::Visible : EVisibility::Hidden;
+	if (URuntimeVirtualTextureComponent* Component = RuntimeVirtualTextureComponent.Get())
+	{
+		const bool bVisible = RuntimeVirtualTextureComponent->IsStreamingTextureInvalid();
+		return bVisible ? EVisibility::Visible : EVisibility::Hidden;
+	}
+	return EVisibility::Hidden;
 }
 
 FReply FRuntimeVirtualTextureComponentDetailsCustomization::BuildStreamedMips()
 {
-	// Create a new asset if none is already bound
-	UVirtualTextureBuilder* CreatedTexture = nullptr;
-	if (RuntimeVirtualTextureComponent->GetVirtualTexture() != nullptr && RuntimeVirtualTextureComponent->GetStreamingTexture() == nullptr)
+	if (URuntimeVirtualTextureComponent* Component = RuntimeVirtualTextureComponent.Get())
 	{
-		FAssetToolsModule& AssetToolsModule = FModuleManager::GetModuleChecked<FAssetToolsModule>("AssetTools");
-
-		const FString DefaultPath = FPackageName::GetLongPackagePath(RuntimeVirtualTextureComponent->GetVirtualTexture()->GetPathName());
-		const FString DefaultName = FPackageName::GetShortName(RuntimeVirtualTextureComponent->GetVirtualTexture()->GetName() + TEXT("_SVT"));
+		// Create a new asset if none is already bound
+		UVirtualTextureBuilder* CreatedTexture = nullptr;
+		if (Component->GetVirtualTexture() != nullptr && Component->GetStreamingTexture() == nullptr)
+		{
+			FAssetToolsModule& AssetToolsModule = FModuleManager::GetModuleChecked<FAssetToolsModule>("AssetTools");
 
-		UFactory* Factory = NewObject<UVirtualTextureBuilderFactory>();
-		UObject* Object = AssetToolsModule.Get().CreateAssetWithDialog(DefaultName, DefaultPath, UVirtualTextureBuilder::StaticClass(), Factory);
-		CreatedTexture = Cast<UVirtualTextureBuilder>(Object);
-	}
+			const FString DefaultPath = FPackageName::GetLongPackagePath(Component->GetVirtualTexture()->GetPathName());
+			const FString DefaultName = FPackageName::GetShortName(Component->GetVirtualTexture()->GetName() + TEXT("_SVT"));
 
-	// Build the texture contents
-	bool bOK = false;
-	if (RuntimeVirtualTextureComponent->GetStreamingTexture() != nullptr || CreatedTexture != nullptr)
-	{
-		const FScopedTransaction Transaction(LOCTEXT("Transaction_BuildDebugStreamingTexture", "Build Streaming Texture"));
+			UFactory* Factory = NewObject<UVirtualTextureBuilderFactory>();
+			UObject* Object = AssetToolsModule.Get().CreateAssetWithDialog(DefaultName, DefaultPath, UVirtualTextureBuilder::StaticClass(), Factory);
+			CreatedTexture = Cast<UVirtualTextureBuilder>(Object);
+		}
 
-		if (CreatedTexture != nullptr)
+		// Build the texture contents
+		bool bOK = false;
+		if (Component->GetStreamingTexture() != nullptr || CreatedTexture != nullptr)
 		{
-			RuntimeVirtualTextureComponent->Modify();
-			RuntimeVirtualTextureComponent->SetStreamingTexture(CreatedTexture);
-		}
+			const FScopedTransaction Transaction(LOCTEXT("Transaction_BuildDebugStreamingTexture", "Build Streaming Texture"));
 
-		RuntimeVirtualTextureComponent->GetStreamingTexture()->Modify();
+			if (CreatedTexture != nullptr)
+			{
+				Component->Modify();
+				Component->SetStreamingTexture(CreatedTexture);
+			}
 
-		const FLinearColor FixedColor = RuntimeVirtualTextureComponent->GetStreamingMipsFixedColor();
-		if (RuntimeVirtualTexture::BuildStreamedMips(RuntimeVirtualTextureComponent, FixedColor))
-		{
-			bOK = true;
+			Component->GetStreamingTexture()->Modify();
+
+			const FLinearColor FixedColor = Component->GetStreamingMipsFixedColor();
+			if (RuntimeVirtualTexture::BuildStreamedMips(Component, FixedColor))
+			{
+				bOK = true;
+			}
 		}
-	}
 
-	return bOK ? FReply::Handled() : FReply::Unhandled();
+		return bOK ? FReply::Handled() : FReply::Unhandled();
+	}
+	return FReply::Unhandled();
 }
 
 #undef LOCTEXT_NAMESPACE
