Change Virtual Texture conversion dialog UI.
Remove "apply filter" button, and apply filter whenever we change the filter value in the dropdown.
Select initial filter value as min or max size according to whether we are converting to VT or from VT.
Default initial filter value is fully inclusive so that we don't get strange UX of selecting to convert a set of textures but none of them pass the initial filter.
#test In editor select to convert one or more VT textures to regular (or the other direction) and validate that all selected textures appear in inital filtered view.

diff --git a/Engine/Source/Editor/VirtualTexturingEditor/Private/SConvertToVirtualTexture.cpp b/Engine/Source/Editor/VirtualTexturingEditor/Private/SConvertToVirtualTexture.cpp
--- a/Engine/Source/Editor/VirtualTexturingEditor/Private/SConvertToVirtualTexture.cpp
+++ b/Engine/Source/Editor/VirtualTexturingEditor/Private/SConvertToVirtualTexture.cpp
@@ -31,18 +31,18 @@
 
 #define LOCTEXT_NAMESPACE "SConvertToVirtualTexture"
 
-void SConvertToVirtualTexture::Construct(const FArguments& InArgs)
+void SConvertToVirtualTexture::Construct(const FArguments& InArgs, bool bInBackwards)
 {
 	UserResponse = FConvertToVTDlg::Cancel;
 	ParentWindow = InArgs._ParentWindow.Get();
 	static FName ErrorIcon = "MessageLog.Error";
 
-
 	for (int i = 0; i < 16; i++)
 	{
 		TextureSizes.Add(MakeShareable(new int32(1 << i)));
 	}
-	ThresholdValue = *TextureSizes[10];
+	const int32 InitiallySelectedIndex = bInBackwards ? 15 : 0;
+	ThresholdValue = *TextureSizes[InitiallySelectedIndex];
 
 	this->ChildSlot[
 		SNew(SVerticalBox)
@@ -123,23 +123,12 @@ void SConvertToVirtualTexture::Construct(const FArguments& InArgs)
 					.OptionsSource(&TextureSizes)
 					.OnSelectionChanged(this, &SConvertToVirtualTexture::OnThresholdChanged)
 					.OnGenerateWidget(this, &SConvertToVirtualTexture::OnGenerateThresholdWidget)
-					.InitiallySelectedItem(TextureSizes[10])
+					.InitiallySelectedItem(TextureSizes[InitiallySelectedIndex])
 					[
 						SNew(STextBlock)
 						.Text(this, &SConvertToVirtualTexture::GetThresholdText)
 					]
 				]
-				+ SHorizontalBox::Slot()
-				.AutoWidth()
-				.VAlign(VAlign_Center)
-				[
-					SNew(SButton)
-					.HAlign(HAlign_Center)
-					.ContentPadding(FAppStyle::GetMargin("StandardDialog.ContentPadding"))
-					.OnClicked(this, &SConvertToVirtualTexture::OnFilterButtonClicked)
-					.IsEnabled(this, &SConvertToVirtualTexture::GetFilterButtonEnabled)
-					.Text(LOCTEXT("ConvertToVT_Filter", "Apply Filter"))
-				]
 			]
 		// Separator
 		+ SVerticalBox::Slot()
@@ -180,7 +169,6 @@ void SConvertToVirtualTexture::Construct(const FArguments& InArgs)
 
 	// will be done by SetUserTextures :
 	//UpdateList();
-	bFilterButtonEnabled = false;
 }
 
 void SConvertToVirtualTexture::SetBackwards(bool bSetBackwards)
@@ -526,7 +514,7 @@ FReply SConvertToVirtualTexture::OnButtonClick(FConvertToVTDlg::EResult ButtonID
 void SConvertToVirtualTexture::OnThresholdChanged(TSharedPtr<int32> InSelectedItem, ESelectInfo::Type SelectInfo)
 {
 	ThresholdValue = *InSelectedItem;
-	bFilterButtonEnabled = true;
+	UpdateList();
 }
 
 FText SConvertToVirtualTexture::GetThresholdText() const
@@ -540,19 +528,6 @@ TSharedRef<SWidget> SConvertToVirtualTexture::OnGenerateThresholdWidget(TSharedP
 		.Text(FText::FromString(FString::Format(TEXT("{0}"), TArray<FStringFormatArg>({ *InItem }))));
 }
 
-bool SConvertToVirtualTexture::GetFilterButtonEnabled() const
-{
-	return bFilterButtonEnabled;
-}
-
-FReply SConvertToVirtualTexture::OnFilterButtonClicked()
-{
-	UpdateList();
-	bFilterButtonEnabled = false;
-
-	return FReply::Handled();
-}
-
 FReply SConvertToVirtualTexture::OnExpanderClicked(int index)
 {
 	if (ExpandedIndexes.Contains(index))
@@ -633,7 +608,7 @@ FConvertToVTDlg::FConvertToVTDlg(const TArray<UTexture2D *> &Textures, bool bBac
 			.BorderImage(FAppStyle::GetBrush("ToolPanel.GroupBorder"))
 			.Padding(4.0f)
 			[
-				SAssignNew(DialogWidget, SConvertToVirtualTexture)
+				SAssignNew(DialogWidget, SConvertToVirtualTexture, bBackwards)
 				.ParentWindow(DialogWindow)
 			];
 
diff --git a/Engine/Source/Editor/VirtualTexturingEditor/Private/SConvertToVirtualTexture.h b/Engine/Source/Editor/VirtualTexturingEditor/Private/SConvertToVirtualTexture.h
--- a/Engine/Source/Editor/VirtualTexturingEditor/Private/SConvertToVirtualTexture.h
+++ b/Engine/Source/Editor/VirtualTexturingEditor/Private/SConvertToVirtualTexture.h
@@ -71,7 +71,7 @@ public:
 	SLATE_ATTRIBUTE(TSharedPtr<SWindow>, ParentWindow)
 	SLATE_END_ARGS()
 
-	void Construct(const FArguments& InArgs);
+	void Construct(const FArguments& InArgs, bool bBackwards);
 
 	void SetBackwards(bool bSetBackwards);
 
@@ -144,6 +144,4 @@ private:
 	bool bThresholdVisible;
 
 	TArray<TSharedPtr<int32>> TextureSizes;
-
-	bool bFilterButtonEnabled;
 };
