* Renamed Build / Build Virtual Textures into Build / Build Streaming Virtual Textures to be more consistent with the names in use in the Runtime Virtual Texture Component
* Added explicit reason for why the SVT warning is displayed in order to better guide the user : either the standard SVT is not up to date or the mobile SVT or both

diff --git a/Engine/Source/Editor/LevelEditor/Private/LevelEditorActions.cpp b/Engine/Source/Editor/LevelEditor/Private/LevelEditorActions.cpp
--- a/Engine/Source/Editor/LevelEditor/Private/LevelEditorActions.cpp
+++ b/Engine/Source/Editor/LevelEditor/Private/LevelEditorActions.cpp
@@ -3870,7 +3870,7 @@ void FLevelEditorCommands::RegisterCommands()
 	UI_COMMAND( BuildMinimap, "Build World Partition Editor Minimap", "Builds the minimap for the current world", EUserInterfaceActionType::Button, FInputChord() );
 	UI_COMMAND( BuildLandscapeSplineMeshes, "Build Landscape Spline Meshes", "Builds landscape spline meshes for the current world", EUserInterfaceActionType::Button, FInputChord());
 	UI_COMMAND( BuildTextureStreamingOnly, "Build Texture Streaming", "Build texture streaming data", EUserInterfaceActionType::Button, FInputChord() );
-	UI_COMMAND( BuildVirtualTextureOnly, "Build Virtual Textures", "Build runtime virtual texture low mips streaming data", EUserInterfaceActionType::Button, FInputChord());
+	UI_COMMAND( BuildVirtualTextureOnly, "Build Streaming Virtual Textures", "Build runtime virtual texture low mips streaming data", EUserInterfaceActionType::Button, FInputChord());
 	UI_COMMAND( BuildAllLandscape, "Build Landscape", "Build all data related to landscape (grass maps, physical material, Nanite, dirty height and weight maps)", EUserInterfaceActionType::Button, FInputChord());
 
 	UI_COMMAND( LightingQuality_Production, "Production", "Sets precomputed lighting quality to highest possible quality (slowest computation time.)", EUserInterfaceActionType::RadioButton, FInputChord() );
diff --git a/Engine/Source/Editor/VirtualTexturingEditor/Classes/RuntimeVirtualTextureDetailsCustomization.h b/Engine/Source/Editor/VirtualTexturingEditor/Classes/RuntimeVirtualTextureDetailsCustomization.h
--- a/Engine/Source/Editor/VirtualTexturingEditor/Classes/RuntimeVirtualTextureDetailsCustomization.h
+++ b/Engine/Source/Editor/VirtualTexturingEditor/Classes/RuntimeVirtualTextureDetailsCustomization.h
@@ -65,8 +65,8 @@ protected:
 
 	/** Returns true if BuildStreamedMips button is enabled */
 	bool IsBuildStreamedMipsEnabled() const;
-	/** Do we need to show warning icon to indicate out of date streaming texture */
-	EVisibility IsBuildWarningIconVisible() const;
+	/** Returns true if the streaming texture is up to date */
+	bool IsStreamingTextureValid(FText& OutReason) const;
 	/** Callback for Build button */
 	FReply BuildStreamedMips();
 
diff --git a/Engine/Source/Editor/VirtualTexturingEditor/Private/RuntimeVirtualTextureDetailsCustomization.cpp b/Engine/Source/Editor/VirtualTexturingEditor/Private/RuntimeVirtualTextureDetailsCustomization.cpp
--- a/Engine/Source/Editor/VirtualTexturingEditor/Private/RuntimeVirtualTextureDetailsCustomization.cpp
+++ b/Engine/Source/Editor/VirtualTexturingEditor/Private/RuntimeVirtualTextureDetailsCustomization.cpp
@@ -13,6 +13,7 @@
 #include "RuntimeVirtualTextureBuildStreamingMips.h"
 #include "RuntimeVirtualTextureSetBounds.h"
 #include "ScopedTransaction.h"
+#include "SceneUtils.h"
 #include "SEnumCombo.h"
 #include "SResetToDefaultMenu.h"
 #include "VirtualTextureBuilderFactory.h"
@@ -297,9 +298,10 @@ void FRuntimeVirtualTextureComponentDetailsCustomization::CustomizeDetails(IDeta
 		.IsEnabled(this, &FRuntimeVirtualTextureComponentDetailsCustomization::IsSetBoundsEnabled)
 	];
 
-	FText BuildButtonText = LOCTEXT("Button_Build_Tooltip", "Build the low mips as streaming virtual texture data. \n\
-		If \"Separate Texture For Mobile\" is enabled in the Streaming Texture, only the mobile version of the texture will be updated when hitting this button \
-		while the mobile preview mode is active (and only the desktop version otherwise).");
+	FText BuildButtonText = LOCTEXT("Button_Build_Tooltip", "Build the low mips as streaming virtual texture data for the current shading path. \n"
+		"If \"Separate Texture For Mobile\" is enabled in the Streaming Texture, only the mobile version of the texture will be updated when hitting this button"
+		"while the mobile preview mode is active (and only the desktop version otherwise).\n"
+		"Use the Build / Build Streaming Virtual Textures button to build all versions (of all loaded Runtime Virtual Texture Volumes)");
 
 	// Apply custom widget for BuildStreamingMips.
 	TSharedRef<IPropertyHandle> BuildStreamingMipsPropertyHandle = DetailBuilder.GetProperty(TEXT("bBuildStreamingMipsButton"));
@@ -332,8 +334,21 @@ void FRuntimeVirtualTextureComponentDetailsCustomization::CustomizeDetails(IDeta
 		[
 			SNew(SImage)
 			.Image(FCoreStyle::Get().GetBrush("Icons.Warning"))
-			.Visibility(this, &FRuntimeVirtualTextureComponentDetailsCustomization::IsBuildWarningIconVisible)
-			.ToolTipText(LOCTEXT("Warning_Build_Tooltip", "The settings have changed since the Streaming Texture was last rebuilt. Streaming mips are disabled."))
+			.Visibility_Lambda([this]()
+				{
+					FText Dummy;
+					return IsStreamingTextureValid(Dummy) ? EVisibility::Hidden : EVisibility::Visible;
+				})
+			.ToolTipText_Lambda([this]() 
+			{
+				FText Reason;
+				if (!IsStreamingTextureValid(Reason))
+				{
+					return FText::Format(LOCTEXT("Warning_Build_Tooltip", "The settings have changed since the Streaming Texture was last rebuilt ({0}).\n"
+						"Streaming mips are disabled. Use the Build / Build Streaming Virtual Textures button to rebuild all versions."), Reason);
+				}
+				return FText();
+			})
 		]
 	];
 }
@@ -372,14 +387,23 @@ bool FRuntimeVirtualTextureComponentDetailsCustomization::IsBuildStreamedMipsEna
 	return false;
 }
 
-EVisibility FRuntimeVirtualTextureComponentDetailsCustomization::IsBuildWarningIconVisible() const
+bool FRuntimeVirtualTextureComponentDetailsCustomization::IsStreamingTextureValid(FText& OutReason) const
 {
 	if (URuntimeVirtualTextureComponent* Component = RuntimeVirtualTextureComponent.Get())
 	{
-		const bool bVisible = RuntimeVirtualTextureComponent->IsStreamingTextureInvalid();
-		return bVisible ? EVisibility::Visible : EVisibility::Hidden;
+		if (RuntimeVirtualTextureComponent->IsStreamingTextureInvalid())
+		{
+			const bool bStreamingTextureInvalid = RuntimeVirtualTextureComponent->IsStreamingTextureInvalid(EShadingPath::Deferred);
+			const bool bMobileStreamingTextureInvalid = RuntimeVirtualTextureComponent->IsStreamingTextureInvalid(EShadingPath::Mobile);
+			OutReason = (bStreamingTextureInvalid && bMobileStreamingTextureInvalid)
+				? LOCTEXT("InvalidBothStreamingTexture_Reason", "both Streaming Textures (Mobile and Non-Mobile) are not up to date")
+				: bStreamingTextureInvalid
+				? LOCTEXT("InvalidStreamingTexture_Reason", "the Streaming Texture is not up to date")
+				: LOCTEXT("InvalidMobileStreamingTexture_Reason", "the Mobile Streaming Texture is not up to date");
+			return false;
+		}
 	}
-	return EVisibility::Hidden;
+	return true;
 }
 
 FReply FRuntimeVirtualTextureComponentDetailsCustomization::BuildStreamedMips()
