* Add Build All button for SVT, which bakes the texture for both the current feature level and the alternate one (mobile SVT) if needed
* Centralized code for SVT baking (editor build command, commandlets, RVT component Build(/All) button)
* Better messaging when SVT is detected as not up-to-date (e.g. when the mobile SVT is not up to date because it was only partially rebuilt)
* Use Fixed Color is now consistently used across all SVT baking processes (editor build command, commandlets, RVT component Build(/All) button) instead of only the RVT component Build buttons

diff --git a/Engine/Source/Editor/UnrealEd/Private/EditorBuildUtils.cpp b/Engine/Source/Editor/UnrealEd/Private/EditorBuildUtils.cpp
--- a/Engine/Source/Editor/UnrealEd/Private/EditorBuildUtils.cpp
+++ b/Engine/Source/Editor/UnrealEd/Private/EditorBuildUtils.cpp
@@ -1775,89 +1775,18 @@ bool FEditorBuildUtils::EditorBuildVirtualTexture(UWorld* InWorld)
 		UWorldPartitionRuntimeVirtualTextureBuilder::LoadRuntimeVirtualTextureActors(WorldPartition, ForEachActorWithLoadingResult);
 	}
 
-	// We will need to build VTs for both shading paths
-	const ERHIFeatureLevel::Type CurFeatureLevel = InWorld->GetFeatureLevel();
-	const ERHIFeatureLevel::Type AltFeatureLevel = (CurFeatureLevel == ERHIFeatureLevel::ES3_1 ? GMaxRHIFeatureLevel : ERHIFeatureLevel::ES3_1);
-	const EShadingPath CurShadingPath = FSceneInterface::GetShadingPath(CurFeatureLevel);
-	const EShadingPath AltShadingPath = FSceneInterface::GetShadingPath(AltFeatureLevel);
-
-	TArray<URuntimeVirtualTextureComponent*> Components[2];
-	for (TObjectIterator<URuntimeVirtualTextureComponent> It; It; ++It)
+	// Gather which component to build : 
+	TArray<URuntimeVirtualTextureComponent*> Components;
+	for (TObjectIterator<URuntimeVirtualTextureComponent> It(RF_ClassDefaultObject, false, EInternalObjectFlags::Garbage); It; ++It)
 	{
-		if (Module->HasStreamedMips(CurShadingPath, *It))
-		{
-			Components[0].Add(*It);
-		}
-
-		if (Module->HasStreamedMips(AltShadingPath, *It))
+		if (It->GetWorld() == InWorld)
 		{
-			Components[1].Add(*It);
+			Components.Add(*It);
 		}
 	}
 
-	// Build for a current feature level first
-	int32 NumFeatureLevelsToBuild = 0;
-	FScopedSlowTask BuildTask(static_cast<float>((Components[0].IsEmpty() ? 0 : 1) + (Components[1].IsEmpty() ? 0 : 1)));
-	BuildTask.MakeDialog(true);
-
-	auto EnterProgressForFeatureLevel = [&BuildTask](ERHIFeatureLevel::Type InFeatureLevel)
-	{
-		BuildTask.EnterProgressFrame(1.0f, FText::Format(LOCTEXT("VirtualTextureBuildFeatureLevel", "Building Virtual Textures for {0}"), FText::FromString(LexToString(InFeatureLevel))));
-	};
-
-	EnterProgressForFeatureLevel(CurFeatureLevel);
-	if (!BuildVirtualTextureComponents(Module, CurShadingPath, Components[0]))
-	{
-		return false;
-	}
-
-	// Build for the other feature level if any
-	bool bResult = true;
-	if (Components[1].Num() != 0)
-	{
-		EnterProgressForFeatureLevel(AltFeatureLevel);
-
-		FScopedSlowTask SubBuildTask(3.0f, BuildTask.GetCurrentMessage());
-		{
-			FText SwitchingFeatureLevelText = FText::Format(LOCTEXT("VirtualTextureSwitchToAltFeatureLevel", "Switching feature level to {0}"), FText::FromString(LexToString(AltFeatureLevel)));
-			SubBuildTask.EnterProgressFrame(1.0f, SwitchingFeatureLevelText);
-
-			UMaterialInterface::SetGlobalRequiredFeatureLevel(AltFeatureLevel, true);
-			UMaterial::AllMaterialsCacheResourceShadersForRendering(/*bUpdateProgressDialog = */true, /*bCacheAllRemainingShaders = */true);
-			UMaterialInstance::AllMaterialsCacheResourceShadersForRendering(/*bUpdateProgressDialog = */true, /*bCacheAllRemainingShaders = */true);
-			CompileGlobalShaderMap(AltFeatureLevel);
-
-			InWorld->ChangeFeatureLevel(AltFeatureLevel);
-
-			// Make sure all assets are finished compiling. Recreate render state after shader compilation complete
-			{
-				UMaterialInterface::SubmitRemainingJobsForWorld(InWorld);
-				FAssetCompilingManager::Get().FinishAllCompilation();
-				FAssetCompilingManager::Get().ProcessAsyncTasks();
-				FGlobalComponentRecreateRenderStateContext Context;
-			}
-
-			// Flush all rendering commands issued by UpdateAllPrimitiveSceneInfos inside the FGlobalComponentRecreateRenderStateContext. 
-			// Some rendering commands may trigger some shader compilations that we need to be issued and wait for completion before rendering the RVT.
-			FlushRenderingCommands();
-
-			// FGlobalComponentRecreateRenderStateContext can create new shaderJobs, make sure to wait on them.
-			FAssetCompilingManager::Get().FinishAllCompilation();
-			FAssetCompilingManager::Get().ProcessAsyncTasks();
-
-		}
-
-		SubBuildTask.EnterProgressFrame(1.0f);
-		bResult = BuildVirtualTextureComponents(Module, AltShadingPath, Components[1]);
-
-		// Restore world feature level
-		{
-			SubBuildTask.EnterProgressFrame(1.0f, FText::Format(LOCTEXT("VirtualTextureSwitchBackToFeatureLevel", "Switching back feature level to {0}"), FText::FromString(LexToString(CurFeatureLevel))));
-			UMaterialInterface::SetGlobalRequiredFeatureLevel(CurFeatureLevel, /*bShouldCompile = */ false);
-			InWorld->ChangeFeatureLevel(CurFeatureLevel);
-		}
-	}
-	return bResult;
+	FBuildAllStreamedMipsResult Result = Module->BuildAllStreamedMips({ .World = InWorld, .Components = Components, .bRestoreFeatureLevelAfterBuilding = true });
+	return Result.bSuccess;
 }
 
 void FEditorBuildUtils::EditorBuildAllLandscape(UWorld* InWorld)
diff --git a/Engine/Source/Editor/UnrealEd/Private/WorldPartition/WorldPartitionRuntimeVirtualTextureBuilder.cpp b/Engine/Source/Editor/UnrealEd/Private/WorldPartition/WorldPartitionRuntimeVirtualTextureBuilder.cpp
--- a/Engine/Source/Editor/UnrealEd/Private/WorldPartition/WorldPartitionRuntimeVirtualTextureBuilder.cpp
+++ b/Engine/Source/Editor/UnrealEd/Private/WorldPartition/WorldPartitionRuntimeVirtualTextureBuilder.cpp
@@ -50,106 +50,33 @@ bool UWorldPartitionRuntimeVirtualTextureBuilder::RunInternal(UWorld* World, con
 		LoadRuntimeVirtualTextureActors(WorldPartition, ForEachActorWithLoadingResult);
 	}
 	
-	// Recreate render state after shader compilation complete
-	{
-		FGlobalComponentRecreateRenderStateContext Context;
-	}
-
-	// We will need to build VTs for both shading paths
-	const ERHIFeatureLevel::Type CurFeatureLevel = World->GetFeatureLevel();
-	const ERHIFeatureLevel::Type AltFeatureLevel = (CurFeatureLevel == ERHIFeatureLevel::ES3_1 ? GMaxRHIFeatureLevel : ERHIFeatureLevel::ES3_1);
-	const EShadingPath CurShadingPath = FSceneInterface::GetShadingPath(CurFeatureLevel);
-	const EShadingPath AltShadingPath = FSceneInterface::GetShadingPath(AltFeatureLevel);
-
-	TSet<UPackage*> ModifiedPackages;
-	TArray<URuntimeVirtualTextureComponent*> Components[2];
+	// Gather which component to build : 
+	TArray<URuntimeVirtualTextureComponent*> Components;
 	for (TObjectIterator<URuntimeVirtualTextureComponent> It(RF_ClassDefaultObject, false, EInternalObjectFlags::Garbage); It; ++It)
 	{
 		if (It->GetWorld() == World)
 		{
-			if (VTModule.HasStreamedMips(CurShadingPath, *It))
-			{
-				Components[0].Add(*It);
-			}
-
-			if (VTModule.HasStreamedMips(AltShadingPath, *It))
-			{
-				Components[1].Add(*It);
-			}
+			Components.Add(*It);
 		}
 	}
 
-	// Build for a current feature level first
-	if (Components[0].Num() != 0)
-	{
-		for (URuntimeVirtualTextureComponent* Component : Components[0])
-		{
-			if (VTModule.BuildStreamedMips(CurShadingPath, Component))
-			{
-				if (UVirtualTextureBuilder* VTBuilder = Component->GetStreamingTexture(); VTBuilder->GetPackage()->IsDirty())
-				{
-					ModifiedPackages.Add(VTBuilder->GetPackage());
-				}
-			}
-		}
-	}
-
-	// Build for others if any
-	if (Components[1].Num() != 0)
-	{
-		// Commandlets do not initialize shader resources for alternate feature levels, do it now
-		{
-			bool bUpdateProgressDialog = false;
-			bool bCacheAllRemainingShaders = true;
-			UMaterialInterface::SetGlobalRequiredFeatureLevel(AltFeatureLevel, true);
-			UMaterial::AllMaterialsCacheResourceShadersForRendering(bUpdateProgressDialog, bCacheAllRemainingShaders);
-			UMaterialInstance::AllMaterialsCacheResourceShadersForRendering(bUpdateProgressDialog, bCacheAllRemainingShaders);
-			CompileGlobalShaderMap(AltFeatureLevel);
-		}
-	
-		World->ChangeFeatureLevel(AltFeatureLevel);
-		
-		// Make sure all assets are finished compiling. Recreate render state after shader compilation complete
-		{
-			UMaterialInterface::SubmitRemainingJobsForWorld(World);
-			FAssetCompilingManager::Get().FinishAllCompilation();
-			FAssetCompilingManager::Get().ProcessAsyncTasks();
-			FGlobalComponentRecreateRenderStateContext Context;
-		}
-
-		// Flush all rendering commands issued by UpdateAllPrimitiveSceneInfos inside the FGlobalComponentRecreateRenderStateContext. 
-		// Some rendering commands may trigger some shader compilations that we need to be issued and wait for completion before rendering the RVT.
-		FlushRenderingCommands();
-
-		// FGlobalComponentRecreateRenderStateContext can create new shaderJobs, make sure to wait on them.
-		FAssetCompilingManager::Get().FinishAllCompilation();
-		FAssetCompilingManager::Get().ProcessAsyncTasks();
-				
-		for (URuntimeVirtualTextureComponent* Component : Components[1])
-		{
-			if (VTModule.BuildStreamedMips(AltShadingPath, Component))
-			{
-				if (UVirtualTextureBuilder* VTBuilder = Component->GetStreamingTexture(); VTBuilder->GetPackage()->IsDirty())
-				{
-					ModifiedPackages.Add(VTBuilder->GetPackage());
-				}
-			}
-		}
-	}
-
-	// Restore world feature level
-	World->ChangeFeatureLevel(CurFeatureLevel);
+	FBuildAllStreamedMipsResult Result = VTModule.BuildAllStreamedMips({ .World = World, .Components = Components, .bRestoreFeatureLevelAfterBuilding = false });
 
 	// Wait for VT Textures to be ready before saving
 	FAssetCompilingManager::Get().FinishAllCompilation();
 
-	if (!SavePackages(ModifiedPackages.Array(), PackageHelper, false))
+	if (!Result.bSuccess)
+	{
+		return false;
+	}
+
+	if (!SavePackages(Result.ModifiedPackages.Array(), PackageHelper, false))
 	{
 		return false;
 	}
 
 	TArray<FString> FilesToSubmit;
-	FilesToSubmit.Append(SourceControlHelpers::PackageFilenames(ModifiedPackages.Array()));
+	FilesToSubmit.Append(SourceControlHelpers::PackageFilenames(Result.ModifiedPackages.Array()));
 
 	const FString ChangeDescription = FString::Printf(TEXT("Built RVT for world '%s'"), *World->GetName());
 	return OnFilesModified(FilesToSubmit, ChangeDescription);
diff --git a/Engine/Source/Editor/VirtualTexturingEditor/Classes/RuntimeVirtualTextureDetailsCustomization.h b/Engine/Source/Editor/VirtualTexturingEditor/Classes/RuntimeVirtualTextureDetailsCustomization.h
--- a/Engine/Source/Editor/VirtualTexturingEditor/Classes/RuntimeVirtualTextureDetailsCustomization.h
+++ b/Engine/Source/Editor/VirtualTexturingEditor/Classes/RuntimeVirtualTextureDetailsCustomization.h
@@ -65,10 +65,12 @@ protected:
 
 	/** Returns true if BuildStreamedMips button is enabled */
 	bool IsBuildStreamedMipsEnabled() const;
+	/** Returns the visibility of the BuildAllStreamedMips button */
+	EVisibility GetBuildAllStreamedMipsVisible() const;
 	/** Returns true if the streaming texture is up to date */
 	bool IsStreamingTextureValid(FText& OutReason) const;
 	/** Callback for Build button */
-	FReply BuildStreamedMips();
+	FReply BuildStreamedMips(bool bBuildAll);
 
 	//~ Begin IDetailCustomization Interface.
 	virtual void CustomizeDetails(IDetailLayoutBuilder& DetailBuilder) override;
diff --git a/Engine/Source/Editor/VirtualTexturingEditor/Private/RuntimeVirtualTextureBuildStreamingMips.cpp b/Engine/Source/Editor/VirtualTexturingEditor/Private/RuntimeVirtualTextureBuildStreamingMips.cpp
--- a/Engine/Source/Editor/VirtualTexturingEditor/Private/RuntimeVirtualTextureBuildStreamingMips.cpp
+++ b/Engine/Source/Editor/VirtualTexturingEditor/Private/RuntimeVirtualTextureBuildStreamingMips.cpp
@@ -2,18 +2,28 @@
 
 #include "RuntimeVirtualTextureBuildStreamingMips.h"
 
+#include "AssetCompilingManager.h"
 #include "Components/RuntimeVirtualTextureComponent.h"
+#include "ComponentRecreateRenderStateContext.h"
 #include "ContentStreaming.h"
 #include "Engine/TextureRenderTarget2D.h"
+#include "Engine/World.h"
+#include "Materials/Material.h"
+#include "Materials/MaterialInstance.h"
+#include "Materials/MaterialInterface.h"
 #include "Misc/ScopedSlowTask.h"
+#include "RenderGraphBuilder.h"
 #include "RendererInterface.h"
 #include "RenderTargetPool.h"
 #include "SceneInterface.h"
-#include "RenderGraphBuilder.h"
+#include "SceneUtils.h"
+#include "ShaderCompiler.h"
+#include "UObject/Package.h"
 #include "VT/RuntimeVirtualTexture.h"
 #include "VT/RuntimeVirtualTextureRender.h"
 #include "VT/VirtualTextureBuilder.h"
-#include "SceneUtils.h"
+
+#define LOCTEXT_NAMESPACE "VirtualTexturingEditorModule"
 
 namespace
 {
@@ -124,12 +134,6 @@ namespace
 
 namespace RuntimeVirtualTexture
 {
-	bool HasStreamedMips(URuntimeVirtualTextureComponent* InComponent)
-	{
-		EShadingPath ShadingPath = (InComponent && InComponent->GetScene()) ? InComponent->GetScene()->GetShadingPath() : EShadingPath::Deferred;
-		return HasStreamedMips(ShadingPath, InComponent);
-	}
-
 	bool HasStreamedMips(EShadingPath ShadingPath, URuntimeVirtualTextureComponent* InComponent)
 	{
 		if (InComponent == nullptr)
@@ -154,14 +158,8 @@ namespace RuntimeVirtualTexture
 
 		return true;
 	}
-	
-	bool BuildStreamedMips(URuntimeVirtualTextureComponent* InComponent, FLinearColor const& FixedColor)
-	{
-		EShadingPath ShadingPath = (InComponent && InComponent->GetScene()) ? InComponent->GetScene()->GetShadingPath() : EShadingPath::Deferred;
-		return BuildStreamedMips(ShadingPath, InComponent, FixedColor);
-	}
 
-	bool BuildStreamedMips(EShadingPath ShadingPath, URuntimeVirtualTextureComponent* InComponent, FLinearColor const& FixedColor)
+	bool BuildStreamedMips(EShadingPath ShadingPath, URuntimeVirtualTextureComponent* InComponent)
 	{
 		TRACE_CPUPROFILER_EVENT_SCOPE(BuildStreamedMips);
 
@@ -175,6 +173,7 @@ namespace RuntimeVirtualTexture
 		const uint32 VirtualTextureSceneIndex = RuntimeVirtualTexture::GetRuntimeVirtualTextureSceneIndex_GameThread(InComponent);
 		const FTransform Transform = InComponent->GetComponentTransform();
 		const FBox Bounds = InComponent->Bounds.GetBox();
+		const FLinearColor FixedColor = InComponent->GetStreamingMipsFixedColor();
 
 		FVTProducerDescription VTDesc;
 		RuntimeVirtualTexture->GetProducerDescription(VTDesc, URuntimeVirtualTexture::FInitSettings(), Transform);
@@ -329,4 +328,154 @@ namespace RuntimeVirtualTexture
 
 		return true;
 	}
+
+	FBuildAllStreamedMipsResult BuildAllStreamedMips(const FBuildAllStreamedMipsParams& InParams)
+	{
+		TRACE_CPUPROFILER_EVENT_SCOPE(BuildAllStreamedMips);
+
+		FBuildAllStreamedMipsResult Result;
+
+		// We will need to build VTs for both shading paths
+		const ERHIFeatureLevel::Type CurFeatureLevel = InParams.World->GetFeatureLevel();
+		const ERHIFeatureLevel::Type AltFeatureLevel = (CurFeatureLevel == ERHIFeatureLevel::ES3_1 ? GMaxRHIFeatureLevel : ERHIFeatureLevel::ES3_1);
+		const EShadingPath CurShadingPath = FSceneInterface::GetShadingPath(CurFeatureLevel);
+		const EShadingPath AltShadingPath = FSceneInterface::GetShadingPath(AltFeatureLevel);
+
+		TArray<URuntimeVirtualTextureComponent*> Components[2];
+		for (URuntimeVirtualTextureComponent* Component : InParams.Components)
+		{
+			if (Component->GetWorld() == InParams.World)
+			{
+				if (HasStreamedMips(CurShadingPath, Component))
+				{
+					Components[0].Add(Component);
+				}
+
+				if (HasStreamedMips(AltShadingPath, Component))
+				{
+					Components[1].Add(Component);
+				}
+			}
+		}
+
+		const int32 NumStreamedMips = Components[0].Num() + Components[1].Num();
+		const int32 NumSteps = 
+			/* Initial flush = */1 
+			+ NumStreamedMips 
+			/* Switch to alternate feature level and back = */ + (Components[1].IsEmpty() ? 0 : 2);
+
+		{
+			FScopedSlowTask Task(static_cast<float>(NumSteps), LOCTEXT("BuildingStreamingMips", "Building Streamed Mips"));
+			Task.MakeDialog(true);
+
+			{
+				// Initial flush : 
+				Task.EnterProgressFrame();
+
+				// Make sure the World is fully streamed in and ready to render at the current feature level :
+				InParams.World->FlushLevelStreaming();
+				FAssetCompilingManager::Get().FinishAllCompilation();
+
+				// Recreate render state after shader compilation complete
+				{
+					FGlobalComponentRecreateRenderStateContext Context;
+				}
+			}
+
+			// Build for a current feature level first
+			if (Components[0].Num() != 0)
+			{
+				for (URuntimeVirtualTextureComponent* Component : Components[0])
+				{
+					if (Task.ShouldCancel())
+					{
+						Result.bSuccess = false;
+						break;
+					}
+
+					Task.EnterProgressFrame();
+					bool bSuccess = BuildStreamedMips(CurShadingPath, Component);
+					if (bSuccess)
+					{
+						if (UVirtualTextureBuilder* VTBuilder = Component->GetStreamingTexture(); VTBuilder->GetPackage()->IsDirty())
+						{
+							Result.ModifiedPackages.Add(VTBuilder->GetPackage());
+						}
+					}
+
+					Result.bSuccess &= bSuccess;
+				}
+			}
+
+			// Build for others if any
+			if ((Components[1].Num() != 0) && !Task.ShouldCancel())
+			{
+				{
+					// Setup alternate feature level : 
+					Task.EnterProgressFrame();
+
+					// Commandlets do not initialize shader resources for alternate feature levels, do it now
+					{
+						bool bUpdateProgressDialog = false;
+						bool bCacheAllRemainingShaders = true;
+						UMaterialInterface::SetGlobalRequiredFeatureLevel(AltFeatureLevel, true);
+						UMaterial::AllMaterialsCacheResourceShadersForRendering(bUpdateProgressDialog, bCacheAllRemainingShaders);
+						UMaterialInstance::AllMaterialsCacheResourceShadersForRendering(bUpdateProgressDialog, bCacheAllRemainingShaders);
+						CompileGlobalShaderMap(AltFeatureLevel);
+					}
+
+					InParams.World->ChangeFeatureLevel(AltFeatureLevel);
+
+					// Make sure all assets are finished compiling. Recreate render state after shader compilation complete
+					{
+						UMaterialInterface::SubmitRemainingJobsForWorld(InParams.World);
+						FAssetCompilingManager::Get().FinishAllCompilation();
+						FAssetCompilingManager::Get().ProcessAsyncTasks();
+						FGlobalComponentRecreateRenderStateContext Context;
+					}
+
+					// Flush all rendering commands issued by UpdateAllPrimitiveSceneInfos inside the FGlobalComponentRecreateRenderStateContext. 
+					// Some rendering commands may trigger some shader compilations that we need to be issued and wait for completion before rendering the RVT.
+					FlushRenderingCommands();
+
+					// FGlobalComponentRecreateRenderStateContext can create new shaderJobs, make sure to wait on them.
+					FAssetCompilingManager::Get().FinishAllCompilation();
+					FAssetCompilingManager::Get().ProcessAsyncTasks();
+				}
+
+				for (URuntimeVirtualTextureComponent* Component : Components[1])
+				{
+					if (Task.ShouldCancel())
+					{
+						Result.bSuccess = false;
+						break;
+					}
+
+					Task.EnterProgressFrame();
+					bool bSuccess = BuildStreamedMips(AltShadingPath, Component);
+					if (bSuccess)
+					{
+						if (UVirtualTextureBuilder* VTBuilder = Component->GetStreamingTexture(); VTBuilder->GetPackage()->IsDirty())
+						{
+							Result.ModifiedPackages.Add(VTBuilder->GetPackage());
+						}
+					}
+
+					Result.bSuccess &= bSuccess;
+				}
+
+				// Restore world feature level
+				if (InParams.bRestoreFeatureLevelAfterBuilding)
+				{
+					Task.EnterProgressFrame();
+					UMaterialInterface::SetGlobalRequiredFeatureLevel(CurFeatureLevel, /*bShouldCompile = */ false);
+					InParams.World->ChangeFeatureLevel(CurFeatureLevel);
+				}
+			}
+		}
+
+		return MoveTemp(Result);
+	}
 }
+
+#undef LOCTEXT_NAMESPACE
\ No newline at end of file
diff --git a/Engine/Source/Editor/VirtualTexturingEditor/Private/RuntimeVirtualTextureBuildStreamingMips.h b/Engine/Source/Editor/VirtualTexturingEditor/Private/RuntimeVirtualTextureBuildStreamingMips.h
--- a/Engine/Source/Editor/VirtualTexturingEditor/Private/RuntimeVirtualTextureBuildStreamingMips.h
+++ b/Engine/Source/Editor/VirtualTexturingEditor/Private/RuntimeVirtualTextureBuildStreamingMips.h
@@ -4,16 +4,20 @@
 
 #include "CoreMinimal.h"
 
+#include "VirtualTexturingEditorModule.h"
+
 class URuntimeVirtualTextureComponent;
+class UWorld;
 enum class EShadingPath;
 
 namespace RuntimeVirtualTexture
 {
 	/** Returns true if the component describes a runtime virtual texture that has streaming mips. */
-	bool HasStreamedMips(URuntimeVirtualTextureComponent* InComponent);
 	bool HasStreamedMips(EShadingPath ShadingPath, URuntimeVirtualTextureComponent* InComponent);
 
 	/** Build the streaming mips texture. */
-	bool BuildStreamedMips(URuntimeVirtualTextureComponent* InComponent, FLinearColor const& FixedColor);
-	bool BuildStreamedMips(EShadingPath ShadingPath, URuntimeVirtualTextureComponent* InComponent, FLinearColor const& FixedColor);
+	bool BuildStreamedMips(EShadingPath ShadingPath, URuntimeVirtualTextureComponent* InComponent);
+
+	FBuildAllStreamedMipsResult BuildAllStreamedMips(const FBuildAllStreamedMipsParams& InParams);
 };
+
diff --git a/Engine/Source/Editor/VirtualTexturingEditor/Private/RuntimeVirtualTextureDetailsCustomization.cpp b/Engine/Source/Editor/VirtualTexturingEditor/Private/RuntimeVirtualTextureDetailsCustomization.cpp
--- a/Engine/Source/Editor/VirtualTexturingEditor/Private/RuntimeVirtualTextureDetailsCustomization.cpp
+++ b/Engine/Source/Editor/VirtualTexturingEditor/Private/RuntimeVirtualTextureDetailsCustomization.cpp
@@ -13,6 +13,7 @@
 #include "RuntimeVirtualTextureBuildStreamingMips.h"
 #include "RuntimeVirtualTextureSetBounds.h"
 #include "ScopedTransaction.h"
+#include "SceneInterface.h"
 #include "SceneUtils.h"
 #include "SEnumCombo.h"
 #include "SResetToDefaultMenu.h"
@@ -298,10 +299,7 @@ void FRuntimeVirtualTextureComponentDetailsCustomization::CustomizeDetails(IDeta
 		.IsEnabled(this, &FRuntimeVirtualTextureComponentDetailsCustomization::IsSetBoundsEnabled)
 	];
 
-	FText BuildButtonText = LOCTEXT("Button_Build_Tooltip", "Build the low mips as streaming virtual texture data for the current shading path. \n"
-		"If \"Separate Texture For Mobile\" is enabled in the Streaming Texture, only the mobile version of the texture will be updated when hitting this button"
-		"while the mobile preview mode is active (and only the desktop version otherwise).\n"
-		"Use the Build / Build Streaming Virtual Textures button to build all versions (of all loaded Runtime Virtual Texture Volumes)");
+	static const FText BuildButtonSuffixText = LOCTEXT("Button_Build_Suffix_Tooltip", "Use the Build / Build Streaming Virtual Textures menu to build all versions of all Runtime Virtual Texture Volumes, after ensuring all actors rendering to RVT are loaded.");
 
 	// Apply custom widget for BuildStreamingMips.
 	TSharedRef<IPropertyHandle> BuildStreamingMipsPropertyHandle = DetailBuilder.GetProperty(TEXT("bBuildStreamingMipsButton"));
@@ -311,24 +309,42 @@ void FRuntimeVirtualTextureComponentDetailsCustomization::CustomizeDetails(IDeta
 		SNew(STextBlock)
 		.Font(IDetailLayoutBuilder::GetDetailFont())
 		.Text(LOCTEXT("Button_BuildStreamingTexture", "Build Streaming Texture"))
-		.ToolTipText(BuildButtonText)
+		.ToolTipText(FText::Format(LOCTEXT("Button_Build_Prefix_Tooltip", "Build the low mips as streaming virtual texture data.\n{0}"), BuildButtonSuffixText))
 	]
 	.ValueContent()
+	.MinDesiredWidth(150.f)
 	[
 		SNew(SHorizontalBox)
 		+ SHorizontalBox::Slot()
 		.FillWidth(4.0f)
+		.MinWidth(100.0f)
 		[
 			SNew(SButton)
 			.HAlign(HAlign_Center)
 			.VAlign(VAlign_Center)
 			.ContentPadding(2.f)
 			.Text(LOCTEXT("Button_Build", "Build"))
-			.ToolTipText(BuildButtonText)
-			.OnClicked(this, &FRuntimeVirtualTextureComponentDetailsCustomization::BuildStreamedMips)
+			.ToolTipText(FText::Format(LOCTEXT("Button_Build_Tooltip", "Build the low mips as streaming virtual texture data for the current shading path, using the currently loaded actors. \n"
+				"If \"Separate Texture For Mobile\" is enabled in the Streaming Texture, only the mobile version of the texture will be updated when using this button "
+				"while the mobile preview mode is active (and only the desktop version otherwise).\n{0}"), BuildButtonSuffixText))
+			.OnClicked(this, &FRuntimeVirtualTextureComponentDetailsCustomization::BuildStreamedMips, /*bBuildAll = */false)
 			.IsEnabled(this, &FRuntimeVirtualTextureComponentDetailsCustomization::IsBuildStreamedMipsEnabled)
 		]
 		+ SHorizontalBox::Slot()
+		.FillWidth(4.0f)
+		.MinWidth(100.0f)
+		[
+			SNew(SButton)
+			.HAlign(HAlign_Center)
+			.VAlign(VAlign_Center)
+			.ContentPadding(2.f)
+			.Text(LOCTEXT("Button_BuildAll", "Build All"))
+			.ToolTipText(FText::Format(LOCTEXT("Button_BuildAll_Tooltip", "Build the low mips as streaming virtual texture data for all shading paths, if necessary, using the currently loaded actors. \n{0}"), BuildButtonSuffixText))
+			.OnClicked(this, &FRuntimeVirtualTextureComponentDetailsCustomization::BuildStreamedMips, /*bBuildAll = */true)
+			.IsEnabled(this, &FRuntimeVirtualTextureComponentDetailsCustomization::IsBuildStreamedMipsEnabled)
+			.Visibility(this, &FRuntimeVirtualTextureComponentDetailsCustomization::GetBuildAllStreamedMipsVisible)
+		]
+		+ SHorizontalBox::Slot()
 		.HAlign(HAlign_Center)
 		.VAlign(VAlign_Center)
 		[
@@ -344,8 +360,12 @@ void FRuntimeVirtualTextureComponentDetailsCustomization::CustomizeDetails(IDeta
 				FText Reason;
 				if (!IsStreamingTextureValid(Reason))
 				{
-					return FText::Format(LOCTEXT("Warning_Build_Tooltip", "The settings have changed since the Streaming Texture was last rebuilt ({0}).\n"
-						"Streaming mips are disabled. Use the Build / Build Streaming Virtual Textures button to rebuild all versions."), Reason);
+					const FText BuildAllShadingPathsText = LOCTEXT("BuildAllShadingPaths", " or Build All to rebuild it for all shading paths");
+					return FText::Format(LOCTEXT("Warning_Build_Tooltip", "The settings have changed since the Streaming Texture was last rebuilt.\n{0}\n"
+						"Use the Build button to rebuild the streaming virtual texture for the current shading path{1}.\n"
+						"Build Streaming Virtual Textures in the Build menu can also be used to rebuild all loaded streaming virtual textures for all shading paths, after ensuring all actors rendering to RVT are loaded.\n"
+						"Meanwhile, streaming mips will be disabled."), 
+						Reason, (GetBuildAllStreamedMipsVisible() == EVisibility::Visible) ? BuildAllShadingPathsText : FText());
 				}
 				return FText();
 			})
@@ -387,29 +407,58 @@ bool FRuntimeVirtualTextureComponentDetailsCustomization::IsBuildStreamedMipsEna
 	return false;
 }
 
+EVisibility FRuntimeVirtualTextureComponentDetailsCustomization::GetBuildAllStreamedMipsVisible() const
+{
+	if (URuntimeVirtualTextureComponent* Component = RuntimeVirtualTextureComponent.Get())
+	{
+		UVirtualTextureBuilder* StreamedMipsTexture = Component->GetStreamingTexture();
+		// If there is a virtual texture to build for another shading path, then Build All makes sense
+		if (IsBuildStreamedMipsEnabled() && (StreamedMipsTexture != nullptr) && StreamedMipsTexture->bSeparateTextureForMobile)
+		{
+			return EVisibility::Visible;
+		}
+	}
+	return EVisibility::Collapsed;
+}
+
 bool FRuntimeVirtualTextureComponentDetailsCustomization::IsStreamingTextureValid(FText& OutReason) const
 {
+	using EStreamingTextureStatusFlags = URuntimeVirtualTextureComponent::EStreamingTextureStatusFlags;
+
 	if (URuntimeVirtualTextureComponent* Component = RuntimeVirtualTextureComponent.Get())
 	{
 		if (RuntimeVirtualTextureComponent->IsStreamingTextureInvalid())
 		{
-			const bool bStreamingTextureInvalid = RuntimeVirtualTextureComponent->IsStreamingTextureInvalid(EShadingPath::Deferred);
-			const bool bMobileStreamingTextureInvalid = RuntimeVirtualTextureComponent->IsStreamingTextureInvalid(EShadingPath::Mobile);
-			OutReason = (bStreamingTextureInvalid && bMobileStreamingTextureInvalid)
-				? LOCTEXT("InvalidBothStreamingTexture_Reason", "both Streaming Textures (Mobile and Non-Mobile) are not up to date")
-				: bStreamingTextureInvalid
-				? LOCTEXT("InvalidStreamingTexture_Reason", "the Streaming Texture is not up to date")
-				: LOCTEXT("InvalidMobileStreamingTexture_Reason", "the Mobile Streaming Texture is not up to date");
+			FTextBuilder TextBuilder;
+			auto AppendError = [&TextBuilder](EStreamingTextureStatusFlags InStatus, const FText& InStreamingTextureName)
+				{
+					if (EnumHasAllFlags(InStatus, EStreamingTextureStatusFlags::HasVirtualTexture | EStreamingTextureStatusFlags::HasStreamingTexture)
+						&& EnumHasAnyFlags(InStatus, EStreamingTextureStatusFlags::InvalidStreamingTexture))
+					{
+						TextBuilder.AppendLineFormat(LOCTEXT("UnbuiltStreamingTexture_Reason", "The {0} is setup but not built yet."), InStreamingTextureName);
+					}
+					if (EnumHasAllFlags(InStatus, EStreamingTextureStatusFlags::HasVirtualTexture | EStreamingTextureStatusFlags::HasStreamingTexture)
+						&& EnumHasAnyFlags(InStatus, EStreamingTextureStatusFlags::NonMatchingStreamingTextureSettings))
+					{
+						TextBuilder.AppendLineFormat(LOCTEXT("NotUpToDateStreamingTexture_Reason", "The {0} is not up to date."), InStreamingTextureName);
+					}
+				};
+			AppendError(RuntimeVirtualTextureComponent->GetStreamingTextureStatus(EShadingPath::Deferred), LOCTEXT("StreamingTexture", "streaming texture"));
+			AppendError(RuntimeVirtualTextureComponent->GetStreamingTextureStatus(EShadingPath::Mobile), LOCTEXT("MobileStreamingTexture", "mobile streaming texture"));
+			OutReason = TextBuilder.ToText();
 			return false;
 		}
 	}
 	return true;
 }
 
-FReply FRuntimeVirtualTextureComponentDetailsCustomization::BuildStreamedMips()
+FReply FRuntimeVirtualTextureComponentDetailsCustomization::BuildStreamedMips(bool bBuildAll)
 {
 	if (URuntimeVirtualTextureComponent* Component = RuntimeVirtualTextureComponent.Get())
 	{
+		UWorld* World = Component->GetWorld();
+		check(World != nullptr);
+
 		// Create a new asset if none is already bound
 		UVirtualTextureBuilder* CreatedTexture = nullptr;
 		if (Component->GetVirtualTexture() != nullptr && Component->GetStreamingTexture() == nullptr)
@@ -438,10 +487,16 @@ FReply FRuntimeVirtualTextureComponentDetailsCustomization::BuildStreamedMips()
 
 			Component->GetStreamingTexture()->Modify();
 
-			const FLinearColor FixedColor = Component->GetStreamingMipsFixedColor();
-			if (RuntimeVirtualTexture::BuildStreamedMips(Component, FixedColor))
+			if (bBuildAll)
+			{
+				FBuildAllStreamedMipsResult Result = RuntimeVirtualTexture::BuildAllStreamedMips({ .World = World, .Components = { Component }, .bRestoreFeatureLevelAfterBuilding = true });
+				bOK = Result.bSuccess;
+			}
+			else
 			{
-				bOK = true;
+				const ERHIFeatureLevel::Type CurFeatureLevel = World->GetFeatureLevel();
+				const EShadingPath CurShadingPath = FSceneInterface::GetShadingPath(CurFeatureLevel);
+				bOK = RuntimeVirtualTexture::BuildStreamedMips(CurShadingPath, Component);
 			}
 		}
 
diff --git a/Engine/Source/Editor/VirtualTexturingEditor/Private/VirtualTexturingEditorModule.cpp b/Engine/Source/Editor/VirtualTexturingEditor/Private/VirtualTexturingEditorModule.cpp
--- a/Engine/Source/Editor/VirtualTexturingEditor/Private/VirtualTexturingEditorModule.cpp
+++ b/Engine/Source/Editor/VirtualTexturingEditor/Private/VirtualTexturingEditorModule.cpp
@@ -32,10 +32,13 @@ public:
 	//~ End IModuleInterface Interface.
 
 	//~ Begin IVirtualTexturingEditorModule Interface.
+	UE_DEPRECATED(5.6, "Use the version that takes a EShadingPath (previously was : EShadingPath::Deferred)")
 	virtual bool HasStreamedMips(URuntimeVirtualTextureComponent* InComponent) const override;
-	virtual bool BuildStreamedMips(URuntimeVirtualTextureComponent* InComponent) const override;
 	virtual bool HasStreamedMips(EShadingPath ShadingPath, URuntimeVirtualTextureComponent* InComponent) const override;
+	UE_DEPRECATED(5.6, "Use the version that takes a EShadingPath (previously was : EShadingPath::Deferred)")
+	virtual bool BuildStreamedMips(URuntimeVirtualTextureComponent* InComponent) const override;
 	virtual bool BuildStreamedMips(EShadingPath ShadingPath, URuntimeVirtualTextureComponent* InComponent) const override;
+	virtual FBuildAllStreamedMipsResult BuildAllStreamedMips(const FBuildAllStreamedMipsParams& InParams) const override;
 	virtual void ConvertVirtualTextures(const TArray<UTexture2D *>& Textures, bool bConvertBackToNonVirtual, const TArray<UMaterial *>* RelatedMaterials /* = nullptr */) const override;
 	virtual void ConvertVirtualTexturesWithDialog(const TArray<UTexture2D *>& Textures, bool bConvertBackToNonVirtual) const override;
 	//~ End IVirtualTexturingEditorModule Interface.
@@ -87,24 +90,31 @@ void FVirtualTexturingEditorModule::OnPlacementModeRefresh(FName CategoryName)
 	}
 }
 
+// Deprecated
 bool FVirtualTexturingEditorModule::HasStreamedMips(URuntimeVirtualTextureComponent* InComponent) const
 {
 	return RuntimeVirtualTexture::HasStreamedMips(EShadingPath::Deferred, InComponent);
 }
 
-bool FVirtualTexturingEditorModule::BuildStreamedMips(URuntimeVirtualTextureComponent* InComponent) const
+bool FVirtualTexturingEditorModule::HasStreamedMips(EShadingPath ShadingPath, URuntimeVirtualTextureComponent* InComponent) const
 {
-	return RuntimeVirtualTexture::BuildStreamedMips(EShadingPath::Deferred, InComponent, FLinearColor::Transparent);
+	return RuntimeVirtualTexture::HasStreamedMips(ShadingPath, InComponent);
 }
 
-bool FVirtualTexturingEditorModule::HasStreamedMips(EShadingPath ShadingPath, URuntimeVirtualTextureComponent* InComponent) const
+// Deprecated
+bool FVirtualTexturingEditorModule::BuildStreamedMips(URuntimeVirtualTextureComponent* InComponent) const
 {
-	return RuntimeVirtualTexture::HasStreamedMips(ShadingPath, InComponent);
+	return RuntimeVirtualTexture::BuildStreamedMips(EShadingPath::Deferred, InComponent);
 }
 
 bool FVirtualTexturingEditorModule::BuildStreamedMips(EShadingPath ShadingPath, URuntimeVirtualTextureComponent* InComponent) const
 {
-	return RuntimeVirtualTexture::BuildStreamedMips(ShadingPath, InComponent, FLinearColor::Transparent);
+	return RuntimeVirtualTexture::BuildStreamedMips(ShadingPath, InComponent);
+}
+
+FBuildAllStreamedMipsResult FVirtualTexturingEditorModule::BuildAllStreamedMips(const FBuildAllStreamedMipsParams& InParams) const
+{
+	return RuntimeVirtualTexture::BuildAllStreamedMips(InParams);
 }
 
 void FVirtualTexturingEditorModule::ConvertVirtualTextures(const TArray<UTexture2D *>& Textures, bool bConvertBackToNonVirtual, const TArray<UMaterial *>* RelatedMaterials /* = nullptr */) const
diff --git a/Engine/Source/Editor/VirtualTexturingEditor/Public/VirtualTexturingEditorModule.h b/Engine/Source/Editor/VirtualTexturingEditor/Public/VirtualTexturingEditorModule.h
--- a/Engine/Source/Editor/VirtualTexturingEditor/Public/VirtualTexturingEditorModule.h
+++ b/Engine/Source/Editor/VirtualTexturingEditor/Public/VirtualTexturingEditorModule.h
@@ -7,19 +7,59 @@
 
 class UTexture2D;
 class UMaterial;
+class URuntimeVirtualTextureComponent;
 enum class EShadingPath;
 
+struct FBuildAllStreamedMipsResult
+{
+	bool bSuccess = true;
+	TSet<UPackage*> ModifiedPackages;
+};
+
+struct FBuildAllStreamedMipsParams
+{
+	/** World for which to build the streamed mips */
+	UWorld* World = nullptr;
+	/** RVT components for which to build the streamed mips*/
+	TConstArrayView<URuntimeVirtualTextureComponent*> Components;
+	/** If the feature level has been changed because we needed to build the SVT for multiple feature levels, defines whether it should be restored afterwards */
+	bool bRestoreFeatureLevelAfterBuilding = true;
+};
+
 /** Module for virtual texturing editor extensions. */
 class IVirtualTexturingEditorModule : public IModuleInterface
 {
 public:
-	/** Returns true if the component describes a runtime virtual texture that has streaming low mips. */
-	virtual bool HasStreamedMips(class URuntimeVirtualTextureComponent* InComponent) const = 0;
-	virtual bool HasStreamedMips(EShadingPath ShadingPath, class URuntimeVirtualTextureComponent* InComponent) const = 0;
-	/** Build the contents of the streaming low mips. */
-	virtual bool BuildStreamedMips(class URuntimeVirtualTextureComponent* InComponent) const = 0;
-	virtual bool BuildStreamedMips(EShadingPath ShadingPath, class URuntimeVirtualTextureComponent* InComponent) const = 0;
+	UE_DEPRECATED(5.6, "Use the version that takes a EShadingPath (previously was : EShadingPath::Deferred)")
+	virtual bool HasStreamedMips(URuntimeVirtualTextureComponent* InComponent) const PURE_VIRTUAL(IVirtualTexturingEditorModule::HasStreamedMips, return false;);
+	
+	/**
+	 * Returns true if the component describes a runtime virtual texture that has streaming low mips for the specified shading path
+	 */
+	virtual bool HasStreamedMips(EShadingPath ShadingPath, URuntimeVirtualTextureComponent* InComponent) const PURE_VIRTUAL(IVirtualTexturingEditorModule::HasStreamedMips, return false;);
+	
+	UE_DEPRECATED(5.6, "Use the version that takes a EShadingPath (previously was : EShadingPath::Deferred)")
+	virtual bool BuildStreamedMips(URuntimeVirtualTextureComponent* InComponent) const PURE_VIRTUAL(IVirtualTexturingEditorModule::BuildStreamedMips, return false;);
+	
+	/**
+	 * Build the contents of the streaming low mips for a single component and a single shading path.
+	 * @param ShadingPath shading path for which to build the SVT (note: EShadingPath::Mobile only if enabled)
+	 * @param InComponent RVT component to build the SVT for
+	 * 
+	 * @return true in case of success 
+	 */
+	virtual bool BuildStreamedMips(EShadingPath ShadingPath, URuntimeVirtualTextureComponent* InComponent) const PURE_VIRTUAL(IVirtualTexturingEditorModule::BuildStreamedMips, return false;);
+
+	/**
+	 * Build the contents of the streaming low mips for several components for all supported shading paths (e.g. deferred and mobile, if enabled)
+	 * This is the preferred method for ensuring the SVT is up-to-date across all platforms but is slower because of the potential multiple shading paths
+	 * #param InWorld world to whom belong the RVT components to bake
+	 * @param InComponents components to build the SVT for
+	 * 
+	 * @return a struct containing the details about the whole operation
+	 */
+	virtual FBuildAllStreamedMipsResult BuildAllStreamedMips(const FBuildAllStreamedMipsParams& InParams) const PURE_VIRTUAL(IVirtualTexturingEditorModule::BuildAllStreamedMips, return {};);
 	
-	virtual void ConvertVirtualTextures(const TArray<UTexture2D *>& Textures, bool bConvertBackToNonVirtual, const TArray<UMaterial *>* RelatedMaterials /* = nullptr */) const = 0;
-	virtual void ConvertVirtualTexturesWithDialog(const TArray<UTexture2D*>& Textures, bool bConvertBackToNonVirtual) const = 0;
+	virtual void ConvertVirtualTextures(const TArray<UTexture2D *>& Textures, bool bConvertBackToNonVirtual, const TArray<UMaterial *>* RelatedMaterials /* = nullptr */) const PURE_VIRTUAL(IVirtualTexturingEditorModule::ConvertVirtualTextures, );
+	virtual void ConvertVirtualTexturesWithDialog(const TArray<UTexture2D*>& Textures, bool bConvertBackToNonVirtual) const PURE_VIRTUAL(IVirtualTexturingEditorModule::ConvertVirtualTexturesWithDialog, );
 };
diff --git a/Engine/Source/Runtime/Engine/Classes/Components/RuntimeVirtualTextureComponent.h b/Engine/Source/Runtime/Engine/Classes/Components/RuntimeVirtualTextureComponent.h
--- a/Engine/Source/Runtime/Engine/Classes/Components/RuntimeVirtualTextureComponent.h
+++ b/Engine/Source/Runtime/Engine/Classes/Components/RuntimeVirtualTextureComponent.h
@@ -175,12 +175,25 @@ public:
 	/** Public getter for lossy compression setting. */
 	TEnumAsByte<ETextureLossyCompressionAmount> GetLossyCompressionAmount() const { return LossyCompressionAmount; }
 
-	/** Returns true if there are StreamingTexure contents but they are not valid for use. */
+	/** Returns whether the StreamingTexure contents is valid for use for the specified shading path */
 	ENGINE_API bool IsStreamingTextureInvalid(EShadingPath ShadingPath) const;
 
+	enum class EStreamingTextureStatusFlags : uint8
+	{
+		None = 0,
+		HasVirtualTexture = (1 << 0),
+		HasStreamingTexture = (1 << 1),
+		InvalidStreamingTexture = (1 << 2),
+		NonMatchingStreamingTextureSettings = (1 << 3),
+	};
+	/** Returns the status of the StreamingTexure contents for the specified shading path */
+	ENGINE_API EStreamingTextureStatusFlags GetStreamingTextureStatus(EShadingPath ShadingPath) const;
+
 #if WITH_EDITOR
-	/** Returns true if there are StreamingTexure contents but they are not valid for use, taking into account all rendering modes */
+	/** Returns whether the StreamingTexure contents is valid for use, taking into account all rendering modes */
 	ENGINE_API bool IsStreamingTextureInvalid() const;
+	/**  Returns the status of the StreamingTexure contents, taking into account all rendering modes */
+	ENGINE_API EStreamingTextureStatusFlags GetStreamingTextureStatus() const;
 	/** Set a new asset to hold the low mip streaming texture. This should only be called directly before setting data to the new asset. */
 	void SetStreamingTexture(UVirtualTextureBuilder* InTexture) { StreamingTexture = InTexture; }
 	/** Initialize the low mip streaming texture with the passed in size and data. */
@@ -239,3 +252,5 @@ public:
 	/** Scene proxy object. Managed by the scene but stored here. */
 	class FRuntimeVirtualTextureSceneProxy* SceneProxy;
 };
+
+ENUM_CLASS_FLAGS(URuntimeVirtualTextureComponent::EStreamingTextureStatusFlags);
\ No newline at end of file
diff --git a/Engine/Source/Runtime/Engine/Private/Components/RuntimeVirtualTextureComponent.cpp b/Engine/Source/Runtime/Engine/Private/Components/RuntimeVirtualTextureComponent.cpp
--- a/Engine/Source/Runtime/Engine/Private/Components/RuntimeVirtualTextureComponent.cpp
+++ b/Engine/Source/Runtime/Engine/Private/Components/RuntimeVirtualTextureComponent.cpp
@@ -335,11 +335,39 @@ static bool IsCompatibleFormat(URuntimeVirtualTexture const& RuntimeVirtualTextu
 
 bool URuntimeVirtualTextureComponent::IsStreamingTextureInvalid(EShadingPath ShadingPath) const
 {
-	return 
-		VirtualTexture != nullptr && 
-		StreamingTexture != nullptr && 
-		StreamingTexture->GetVirtualTexture(ShadingPath) != nullptr && 
-		(StreamingTexture->BuildHash != CalculateStreamingTextureSettingsHash() || !IsCompatibleFormat(*VirtualTexture, *StreamingTexture->GetVirtualTexture(ShadingPath)));
+	EStreamingTextureStatusFlags Status = GetStreamingTextureStatus(ShadingPath);
+	return EnumHasAllFlags(Status, EStreamingTextureStatusFlags::HasVirtualTexture | EStreamingTextureStatusFlags::HasStreamingTexture)
+		&& EnumHasAnyFlags(Status, EStreamingTextureStatusFlags::InvalidStreamingTexture | EStreamingTextureStatusFlags::NonMatchingStreamingTextureSettings);
+}
+
+URuntimeVirtualTextureComponent::EStreamingTextureStatusFlags URuntimeVirtualTextureComponent::GetStreamingTextureStatus(EShadingPath ShadingPath) const
+{
+	EStreamingTextureStatusFlags Status = EStreamingTextureStatusFlags::None;
+	if (VirtualTexture != nullptr)
+	{
+		Status |= EStreamingTextureStatusFlags::HasVirtualTexture;
+	}
+	if (StreamingTexture != nullptr)
+	{
+		Status |= EStreamingTextureStatusFlags::HasStreamingTexture;
+	}
+
+	if (EnumHasAllFlags(Status, EStreamingTextureStatusFlags::HasVirtualTexture | EStreamingTextureStatusFlags::HasStreamingTexture))
+	{
+		UVirtualTexture2D* BuiltTexture = ((ShadingPath == EShadingPath::Mobile) && StreamingTexture->bSeparateTextureForMobile) ? StreamingTexture->TextureMobile : StreamingTexture->Texture;
+		if (BuiltTexture == nullptr)
+		{
+			Status |= EStreamingTextureStatusFlags::InvalidStreamingTexture;
+		}
+		else
+		{
+			if (StreamingTexture->BuildHash != CalculateStreamingTextureSettingsHash() || !IsCompatibleFormat(*VirtualTexture, *BuiltTexture))
+			{
+				Status |= EStreamingTextureStatusFlags::NonMatchingStreamingTextureSettings;
+			}
+		}
+	}
+	return Status;
 }
 
 #if WITH_EDITOR
@@ -349,6 +377,11 @@ bool URuntimeVirtualTextureComponent::IsStreamingTextureInvalid() const
 	return IsStreamingTextureInvalid(EShadingPath::Mobile) || IsStreamingTextureInvalid(EShadingPath::Deferred);
 }
 
+URuntimeVirtualTextureComponent::EStreamingTextureStatusFlags URuntimeVirtualTextureComponent::GetStreamingTextureStatus() const
+{
+	return GetStreamingTextureStatus(EShadingPath::Deferred) | GetStreamingTextureStatus(EShadingPath::Mobile);
+}
+
 FLinearColor URuntimeVirtualTextureComponent::GetStreamingMipsFixedColor() const 
 {
 	if (!bUseStreamingMipsFixedColor)
