* Prevent unregistered RVT components from being processed by the SVT builder. This can happen for RVT components part of an unloaded data layer.
* Removed unused function (BuildVirtualTextureComponents)

diff --git a/Engine/Source/Editor/UnrealEd/Private/WorldPartition/WorldPartitionRuntimeVirtualTextureBuilder.cpp b/Engine/Source/Editor/UnrealEd/Private/WorldPartition/WorldPartitionRuntimeVirtualTextureBuilder.cpp
--- a/Engine/Source/Editor/UnrealEd/Private/WorldPartition/WorldPartitionRuntimeVirtualTextureBuilder.cpp
+++ b/Engine/Source/Editor/UnrealEd/Private/WorldPartition/WorldPartitionRuntimeVirtualTextureBuilder.cpp
@@ -50,16 +50,7 @@ bool UWorldPartitionRuntimeVirtualTextureBuilder::RunInternal(UWorld* World, con
 		LoadRuntimeVirtualTextureActors(WorldPartition, ForEachActorWithLoadingResult);
 	}
 	
-	// Gather which component to build : 
-	TArray<URuntimeVirtualTextureComponent*> Components;
-	for (TObjectIterator<URuntimeVirtualTextureComponent> It(RF_ClassDefaultObject, false, EInternalObjectFlags::Garbage); It; ++It)
-	{
-		if (It->GetWorld() == World)
-		{
-			Components.Add(*It);
-		}
-	}
-
+	TArray<URuntimeVirtualTextureComponent*> Components = VTModule.GatherRuntimeVirtualTextureComponents(World);
 	FBuildAllStreamedMipsResult Result = VTModule.BuildAllStreamedMips({ .World = World, .Components = Components, .bRestoreFeatureLevelAfterBuilding = false });
 
 	// Wait for VT Textures to be ready before saving
diff --git a/Engine/Source/Editor/VirtualTexturingEditor/Private/RuntimeVirtualTextureBuildStreamingMips.cpp b/Engine/Source/Editor/VirtualTexturingEditor/Private/RuntimeVirtualTextureBuildStreamingMips.cpp
--- a/Engine/Source/Editor/VirtualTexturingEditor/Private/RuntimeVirtualTextureBuildStreamingMips.cpp
+++ b/Engine/Source/Editor/VirtualTexturingEditor/Private/RuntimeVirtualTextureBuildStreamingMips.cpp
@@ -168,8 +168,16 @@ namespace RuntimeVirtualTexture
 			return true;
 		}
 
+		if (!InComponent->IsRegistered())
+		{
+			UE_LOG(LogVirtualTexturing, Error, TEXT("Trying to build streaming mips for a component (from actor %s) that is not registered. It will be ignored"), *InComponent->GetOwner()->GetActorNameOrLabel());
+			return false;
+		}
+
 		URuntimeVirtualTexture* RuntimeVirtualTexture = InComponent->GetVirtualTexture();
 		FSceneInterface* Scene = InComponent->GetScene();
+		checkf(Scene != nullptr, TEXT("Trying to build streaming mips for a component (from actor %s) that doesn't have a valid scene"), *InComponent->GetOwner()->GetActorNameOrLabel());
+
 		const uint32 VirtualTextureSceneIndex = RuntimeVirtualTexture::GetRuntimeVirtualTextureSceneIndex_GameThread(InComponent);
 		const FTransform Transform = InComponent->GetComponentTransform();
 		const FBox Bounds = InComponent->Bounds.GetBox();
@@ -210,7 +218,7 @@ namespace RuntimeVirtualTexture
 
 		int64 RenderTileResourcesBytes = RenderTileResources.GetTotalSizeBytes();
 
-		UE_LOG(LogVirtualTexturing,Display,TEXT("Allocating %uMiB for RenderTileResourcesBytes"),(uint32)(RenderTileResourcesBytes/(1024*1024)));
+		UE_LOG(LogVirtualTexturing, Display, TEXT("Allocating %uMiB for RenderTileResourcesBytes"),(uint32)(RenderTileResourcesBytes/(1024*1024)));
 
 		// Final pixels will contain image data for each virtual texture layer in order
 		TArray64<uint8> FinalPixels;
@@ -344,7 +352,8 @@ namespace RuntimeVirtualTexture
 		TArray<URuntimeVirtualTextureComponent*> Components[2];
 		for (URuntimeVirtualTextureComponent* Component : InParams.Components)
 		{
-			if (Component->GetWorld() == InParams.World)
+			check(Component->GetWorld() == InParams.World);
+			if (Component->IsRegistered())
 			{
 				if (HasStreamedMips(CurShadingPath, Component))
 				{
@@ -356,6 +365,10 @@ namespace RuntimeVirtualTexture
 					Components[1].Add(Component);
 				}
 			}
+			else
+			{
+				UE_LOG(LogVirtualTexturing, Error, TEXT("Trying to build streaming mips for a component (from actor %s) that is not registered. It will be ignored"), *Component->GetOwner()->GetActorNameOrLabel());
+			}
 		}
 
 		const int32 NumStreamedMips = Components[0].Num() + Components[1].Num();
diff --git a/Engine/Source/Editor/VirtualTexturingEditor/Private/VirtualTexturingEditorModule.cpp b/Engine/Source/Editor/VirtualTexturingEditor/Private/VirtualTexturingEditorModule.cpp
--- a/Engine/Source/Editor/VirtualTexturingEditor/Private/VirtualTexturingEditorModule.cpp
+++ b/Engine/Source/Editor/VirtualTexturingEditor/Private/VirtualTexturingEditorModule.cpp
@@ -2,6 +2,7 @@
 
 #include "VirtualTexturingEditorModule.h"
 
+#include "Components/RuntimeVirtualTextureComponent.h"
 #include "IPlacementModeModule.h"
 #include "MeshPaintVirtualTextureThumbnailRenderer.h"
 #include "PropertyEditorModule.h"
@@ -11,6 +12,7 @@
 #include "RuntimeVirtualTextureThumbnailRenderer.h"
 #include "SConvertToVirtualTexture.h"
 #include "ThumbnailRendering/ThumbnailManager.h"
+#include "UObject/UObjectIterator.h"
 #include "VirtualTextureBuilderAssetTypeActions.h"
 #include "VirtualTextureBuilderThumbnailRenderer.h"
 #include "VT/MeshPaintVirtualTexture.h"
@@ -41,6 +43,7 @@ public:
 	virtual FBuildAllStreamedMipsResult BuildAllStreamedMips(const FBuildAllStreamedMipsParams& InParams) const override;
 	virtual void ConvertVirtualTextures(const TArray<UTexture2D *>& Textures, bool bConvertBackToNonVirtual, const TArray<UMaterial *>* RelatedMaterials /* = nullptr */) const override;
 	virtual void ConvertVirtualTexturesWithDialog(const TArray<UTexture2D *>& Textures, bool bConvertBackToNonVirtual) const override;
+	virtual TArray<URuntimeVirtualTextureComponent*> GatherRuntimeVirtualTextureComponents(UWorld* InWorld) const override;
 	//~ End IVirtualTexturingEditorModule Interface.
 
 private:
@@ -136,4 +139,17 @@ void FVirtualTexturingEditorModule::ConvertVirtualTexturesWithDialog(const TArra
 	SConvertToVirtualTexture::ConvertVTTexture(Textures, bConvertBackToNonVirtual);
 }
 
+TArray<URuntimeVirtualTextureComponent*> FVirtualTexturingEditorModule::GatherRuntimeVirtualTextureComponents(UWorld* InWorld) const
+{
+	TArray<URuntimeVirtualTextureComponent*> Components;
+	for (TObjectIterator<URuntimeVirtualTextureComponent> It(RF_ClassDefaultObject, false, EInternalObjectFlags::Garbage); It; ++It)
+	{
+		if (It->IsRegistered() && (It->GetWorld() == InWorld))
+		{
+			Components.Add(*It);
+		}
+	}
+	return Components;
+}
+
 #undef LOCTEXT_NAMESPACE
diff --git a/Engine/Source/Editor/VirtualTexturingEditor/Public/VirtualTexturingEditorModule.h b/Engine/Source/Editor/VirtualTexturingEditor/Public/VirtualTexturingEditorModule.h
--- a/Engine/Source/Editor/VirtualTexturingEditor/Public/VirtualTexturingEditorModule.h
+++ b/Engine/Source/Editor/VirtualTexturingEditor/Public/VirtualTexturingEditorModule.h
@@ -20,7 +20,7 @@ struct FBuildAllStreamedMipsParams
 {
 	/** World for which to build the streamed mips */
 	UWorld* World = nullptr;
-	/** RVT components for which to build the streamed mips*/
+	/** RVT components for which to build the streamed mips. The components should belong to the World and be properly registered */
 	TConstArrayView<URuntimeVirtualTextureComponent*> Components;
 	/** If the feature level has been changed because we needed to build the SVT for multiple feature levels, defines whether it should be restored afterwards */
 	bool bRestoreFeatureLevelAfterBuilding = true;
@@ -44,7 +44,7 @@ public:
 	/**
 	 * Build the contents of the streaming low mips for a single component and a single shading path.
 	 * @param ShadingPath shading path for which to build the SVT (note: EShadingPath::Mobile only if enabled)
-	 * @param InComponent RVT component to build the SVT for
+	 * @param InComponent RVT component to build the SVT for. It's expected that this component is registered
 	 * 
 	 * @return true in case of success 
 	 */
@@ -53,8 +53,8 @@ public:
 	/**
 	 * Build the contents of the streaming low mips for several components for all supported shading paths (e.g. deferred and mobile, if enabled)
 	 * This is the preferred method for ensuring the SVT is up-to-date across all platforms but is slower because of the potential multiple shading paths
-	 * #param InWorld world to whom belong the RVT components to bake
-	 * @param InComponents components to build the SVT for
+	 * @param InWorld world to whom belong the RVT components to bake
+	 * @param InComponents components to build the SVT for. It's expected that these components are registered
 	 * 
 	 * @return a struct containing the details about the whole operation
 	 */
@@ -62,4 +62,9 @@ public:
 	
 	virtual void ConvertVirtualTextures(const TArray<UTexture2D *>& Textures, bool bConvertBackToNonVirtual, const TArray<UMaterial *>* RelatedMaterials /* = nullptr */) const PURE_VIRTUAL(IVirtualTexturingEditorModule::ConvertVirtualTextures, );
 	virtual void ConvertVirtualTexturesWithDialog(const TArray<UTexture2D*>& Textures, bool bConvertBackToNonVirtual) const PURE_VIRTUAL(IVirtualTexturingEditorModule::ConvertVirtualTexturesWithDialog, );
+
+	/** 
+	 * @return the list of valid RVT components for this world
+	 */
+	virtual TArray<URuntimeVirtualTextureComponent*> GatherRuntimeVirtualTextureComponents(UWorld* InWorld) const PURE_VIRTUAL(IVirtualTexturingEditorModule::GatherRuntimeVirtualTextureComponents, return {};);
 };
