Two fixes to runtime virtual texture snap to landscape.
Use GetCompleteLandscapeExtent() to take into account unloaded landscape areas.
Adjust snapped scale to take into account any additional bounds that are outside of the landscape area.

diff --git a/Engine/Source/Editor/VirtualTexturingEditor/Private/RuntimeVirtualTextureSetBounds.cpp b/Engine/Source/Editor/VirtualTexturingEditor/Private/RuntimeVirtualTextureSetBounds.cpp
--- a/Engine/Source/Editor/VirtualTexturingEditor/Private/RuntimeVirtualTextureSetBounds.cpp
+++ b/Engine/Source/Editor/VirtualTexturingEditor/Private/RuntimeVirtualTextureSetBounds.cpp
@@ -99,23 +99,36 @@ namespace RuntimeVirtualTexture
 				const FVector LandscapePosition = LandscapeTransform.GetTranslation();
 				const FVector LandscapeScale = LandscapeTransform.GetScale3D();
 
+				// We want to set the virtual texture scale so that the landscape quad size is some power of two multiple of the final virtual texture texel size.
 				ULandscapeInfo const* LandscapeInfo = Landscape->GetLandscapeInfo();
-				int32 MinX, MinY, MaxX, MaxY;
-				LandscapeInfo->GetLandscapeExtent(MinX, MinY, MaxX, MaxY);
-				const FIntPoint LandscapeSize(MaxX - MinX + 1, MaxY - MinY + 1);
+				FIntRect LandscapeRect = LandscapeInfo->GetCompleteLandscapeExtent();
+				const FIntPoint LandscapeSize = LandscapeRect.Size();
 				const int32 LandscapeSizeLog2 = FMath::Max(FMath::CeilLogTwo(LandscapeSize.X), FMath::CeilLogTwo(LandscapeSize.Y));
 
 				const int32 VirtualTextureSize = VirtualTexture->GetSize();
 				const int32 VirtualTextureSizeLog2 = FMath::FloorLog2(VirtualTextureSize);
 
-				// Adjust scale.
-				// Note that we need virtual texture resolution to be greater or equal to the landscape resolution for the virtual texture to cover the entire landscape.
 				const int32 VirtualTexelsPerLandscapeVertexLog2 = FMath::Max(VirtualTextureSizeLog2 - LandscapeSizeLog2, 0);
 				const int32 VirtualTexelsPerLandscapeVertex = 1 << VirtualTexelsPerLandscapeVertexLog2;
-				const FVector VirtualTexelWorldSize = LandscapeScale / (float)VirtualTexelsPerLandscapeVertex;
-				const FVector VirtualTextureScale = VirtualTexelWorldSize * (float)VirtualTextureSize;
-				Transform.SetScale3D(FVector(VirtualTextureScale.X, VirtualTextureScale.Y, Transform.GetScale3D().Z));
 
+				FVector VirtualTexelWorldSize = LandscapeScale / (float)VirtualTexelsPerLandscapeVertex;
+				FVector VirtualTextureScale = VirtualTexelWorldSize * (float)VirtualTextureSize;
+
+				// We may need to adjust for the case where the calculated VirtualTextureScale wasn't big enough to cover the previously calculated bounds.
+				const int32 AdditionalScaleX = (int32)FMath::CeilToInt(WorldSize.X / VirtualTextureScale.X);
+				const int32 AdditionalScaleY = (int32)FMath::CeilToInt(WorldSize.Y / VirtualTextureScale.Y);
+				const int32 AdditionalScaleAligned = FMath::RoundUpToPowerOfTwo(FMath::Max(AdditionalScaleX, AdditionalScaleY));
+				// Need to clamp scale so that virtual texture texel is no bigger than a landscape quad. 
+				// This may prevent us from covering the required bounds. In that case the user fix is to increase the resolution of the runtime virtual texture.
+				const float AdditionalScale = (float)FMath::Min(VirtualTexelsPerLandscapeVertex, AdditionalScaleAligned);
+
+				VirtualTexelWorldSize.X *= AdditionalScale;
+				VirtualTexelWorldSize.Y *= AdditionalScale;
+				VirtualTextureScale.X *= AdditionalScale;
+				VirtualTextureScale.Y *= AdditionalScale;
+
+				Transform.SetScale3D(FVector(VirtualTextureScale.X, VirtualTextureScale.Y, Transform.GetScale3D().Z));
+				
 				// Adjust position to snap at a half texel offset from landscape.
 				const FVector BaseVirtualTexturePosition = Transform.GetTranslation();
 				const FVector LandscapeSnapPosition = LandscapePosition - 0.5f * VirtualTexelWorldSize;
