Material Cache, added tag asset actions

diff --git a/Engine/Source/Editor/VirtualTexturingEditor/Classes/MaterialCacheVirtualTextureTagAssetTypeActions.h b/Engine/Source/Editor/VirtualTexturingEditor/Classes/MaterialCacheVirtualTextureTagAssetTypeActions.h
new file mode 100644
--- /dev/null
+++ b/Engine/Source/Editor/VirtualTexturingEditor/Classes/MaterialCacheVirtualTextureTagAssetTypeActions.h
@@ -0,0 +1,21 @@
+// Copyright Epic Games, Inc. All Rights Reserved.
+
+#pragma once
+
+#include "CoreMinimal.h"
+#include "AssetTypeActions_Base.h"
+
+class FAssetTypeActions_MaterialCacheVirtualTextureTag : public FAssetTypeActions_Base
+{
+public:
+	FAssetTypeActions_MaterialCacheVirtualTextureTag() {}
+
+protected:
+	//~ Begin FAssetTypeActions_Base Interface.
+	virtual UClass* GetSupportedClass() const override;
+	virtual FText GetName() const override;
+	virtual FColor GetTypeColor() const override;
+	virtual uint32 GetCategories() override;
+	virtual void GetActions(TArray<UObject*> const& InObjects, FMenuBuilder& MenuBuilder) override;
+	//~ End FAssetTypeActions_Base Interface.
+};
diff --git a/Engine/Source/Editor/VirtualTexturingEditor/Classes/MaterialCacheVirtualTextureTagFactory.h b/Engine/Source/Editor/VirtualTexturingEditor/Classes/MaterialCacheVirtualTextureTagFactory.h
new file mode 100644
--- /dev/null
+++ b/Engine/Source/Editor/VirtualTexturingEditor/Classes/MaterialCacheVirtualTextureTagFactory.h
@@ -0,0 +1,18 @@
+// Copyright Epic Games, Inc. All Rights Reserved.
+
+#pragma once
+
+#include "Factories/Factory.h"
+#include "MaterialCacheVirtualTextureTagFactory.generated.h"
+
+class FFeedbackContext;
+class UClass;
+class UObject;
+
+UCLASS(hidecategories = (Object))
+class UMaterialCacheVirtualTextureTagFactory : public UFactory
+{
+	GENERATED_UCLASS_BODY()
+
+	virtual UObject* FactoryCreateNew(UClass* Class, UObject* InParent, FName Name, EObjectFlags Flags, UObject* Context, FFeedbackContext* Warn) override;
+};
diff --git a/Engine/Source/Editor/VirtualTexturingEditor/Private/MaterialCacheVirtualTextureTagAssetTypeActions.cpp b/Engine/Source/Editor/VirtualTexturingEditor/Private/MaterialCacheVirtualTextureTagAssetTypeActions.cpp
new file mode 100644
--- /dev/null
+++ b/Engine/Source/Editor/VirtualTexturingEditor/Private/MaterialCacheVirtualTextureTagAssetTypeActions.cpp
@@ -0,0 +1,34 @@
+// Copyright Epic Games, Inc. All Rights Reserved.
+
+#include "MaterialCacheVirtualTextureTagAssetTypeActions.h"
+#include "ContentBrowserModule.h"
+#include "MaterialCache/MaterialCacheVirtualTextureTag.h"
+
+#define LOCTEXT_NAMESPACE "MaterialCacheVirtualTextureTagEditorModule"
+
+UClass* FAssetTypeActions_MaterialCacheVirtualTextureTag::GetSupportedClass() const
+{
+	return UMaterialCacheVirtualTextureTag::StaticClass();
+}
+
+FText FAssetTypeActions_MaterialCacheVirtualTextureTag::GetName() const
+{
+	return LOCTEXT("AssetTypeActions_MaterialCacheVirtualTextureTag", "Material Cache Virtual Texture Tag"); 
+}
+
+FColor FAssetTypeActions_MaterialCacheVirtualTextureTag::GetTypeColor() const 
+{
+	return FColor(255, 100, 100); 
+}
+
+uint32 FAssetTypeActions_MaterialCacheVirtualTextureTag::GetCategories() 
+{
+	return EAssetTypeCategories::Textures; 
+}
+
+void FAssetTypeActions_MaterialCacheVirtualTextureTag::GetActions(TArray<UObject*> const& InObjects, FMenuBuilder& MenuBuilder)
+{
+	
+}
+
+#undef LOCTEXT_NAMESPACE
diff --git a/Engine/Source/Editor/VirtualTexturingEditor/Private/MaterialCacheVirtualTextureTagFactory.cpp b/Engine/Source/Editor/VirtualTexturingEditor/Private/MaterialCacheVirtualTextureTagFactory.cpp
new file mode 100644
--- /dev/null
+++ b/Engine/Source/Editor/VirtualTexturingEditor/Private/MaterialCacheVirtualTextureTagFactory.cpp
@@ -0,0 +1,23 @@
+// Copyright Epic Games, Inc. All Rights Reserved.
+
+#include "MaterialCacheVirtualTextureTagFactory.h"
+#include "MaterialCache/MaterialCacheVirtualTextureTag.h"
+
+#include UE_INLINE_GENERATED_CPP_BY_NAME(MaterialCacheVirtualTextureTagFactory)
+
+UMaterialCacheVirtualTextureTagFactory::UMaterialCacheVirtualTextureTagFactory(const FObjectInitializer& ObjectInitializer)
+	: Super(ObjectInitializer)
+{
+	SupportedClass = UMaterialCacheVirtualTextureTag::StaticClass();
+	
+	bCreateNew    = true;
+	bEditAfterNew = true;
+	bEditorImport = false;
+}
+
+UObject* UMaterialCacheVirtualTextureTagFactory::FactoryCreateNew(UClass* Class, UObject* InParent, FName Name, EObjectFlags Flags, UObject* Context, FFeedbackContext* Warn)
+{
+	UMaterialCacheVirtualTextureTag* Tag = NewObject<UMaterialCacheVirtualTextureTag>(InParent, Class, Name, Flags);
+	check(Tag);
+	return Tag;
+}
diff --git a/Engine/Source/Editor/VirtualTexturingEditor/Private/VirtualTexturingEditorModule.cpp b/Engine/Source/Editor/VirtualTexturingEditor/Private/VirtualTexturingEditorModule.cpp
--- a/Engine/Source/Editor/VirtualTexturingEditor/Private/VirtualTexturingEditorModule.cpp
+++ b/Engine/Source/Editor/VirtualTexturingEditor/Private/VirtualTexturingEditorModule.cpp
@@ -4,6 +4,7 @@
 
 #include "Components/RuntimeVirtualTextureComponent.h"
 #include "IPlacementModeModule.h"
+#include "MaterialCacheVirtualTextureTagAssetTypeActions.h"
 #include "MeshPaintVirtualTextureThumbnailRenderer.h"
 #include "PropertyEditorModule.h"
 #include "RuntimeVirtualTextureAssetTypeActions.h"
@@ -57,6 +58,7 @@ void FVirtualTexturingEditorModule::StartupModule()
 	IAssetTools& AssetTools = FModuleManager::LoadModuleChecked<FAssetToolsModule>("AssetTools").Get();
 	AssetTools.RegisterAssetTypeActions(MakeShareable(new FAssetTypeActions_RuntimeVirtualTexture));
 	AssetTools.RegisterAssetTypeActions(MakeShareable(new FAssetTypeActions_VirtualTextureBuilder));
+	AssetTools.RegisterAssetTypeActions(MakeShareable(new FAssetTypeActions_MaterialCacheVirtualTextureTag));
 
 	FPropertyEditorModule& PropertyModule = FModuleManager::GetModuleChecked<FPropertyEditorModule>("PropertyEditor");
 	PropertyModule.RegisterCustomClassLayout("RuntimeVirtualTexture", FOnGetDetailCustomizationInstance::CreateStatic(&FRuntimeVirtualTextureDetailsCustomization::MakeInstance));
