RVT
* The display name of the "Virtual Texture Builder" (also sometimes previously known as "SVT") asset is now "Streaming Runtime Virtual Texture"  (i.e. SRVT)
* It can now be created from the content browser
* Prevent StreamingTexture from being carried over when duplicating a RVT component (Save As, copy/paste...) : that avoids the common mistake of using another map's SRVT after a Save As

diff --git a/Engine/Source/Editor/LevelEditor/Private/LevelEditorActions.cpp b/Engine/Source/Editor/LevelEditor/Private/LevelEditorActions.cpp
--- a/Engine/Source/Editor/LevelEditor/Private/LevelEditorActions.cpp
+++ b/Engine/Source/Editor/LevelEditor/Private/LevelEditorActions.cpp
@@ -4176,7 +4176,7 @@ void FLevelEditorCommands::RegisterCommands()
 	UI_COMMAND( BuildMinimap, "Build World Partition Editor Minimap", "Builds the minimap for the current world", EUserInterfaceActionType::Button, FInputChord() );
 	UI_COMMAND( BuildLandscapeSplineMeshes, "Build Landscape Spline Meshes", "Builds landscape spline meshes for the current world", EUserInterfaceActionType::Button, FInputChord());
 	UI_COMMAND( BuildTextureStreamingOnly, "Build Texture Streaming", "Build texture streaming data", EUserInterfaceActionType::Button, FInputChord() );
-	UI_COMMAND( BuildVirtualTextureOnly, "Build Streaming Virtual Textures", "Build runtime virtual texture low mips streaming data", EUserInterfaceActionType::Button, FInputChord());
+	UI_COMMAND( BuildVirtualTextureOnly, "Build Streaming Runtime Virtual Textures", "Build runtime virtual texture low mips streaming data", EUserInterfaceActionType::Button, FInputChord());
 	UI_COMMAND( BuildAllLandscape, "Build Landscape", "Build all data related to landscape (grass maps, physical material, Nanite, dirty height and weight maps)", EUserInterfaceActionType::Button, FInputChord());
 
 	UI_COMMAND( LightingQuality_Production, "Production", "Sets precomputed lighting quality to highest possible quality (slowest computation time.)", EUserInterfaceActionType::RadioButton, FInputChord() );
diff --git a/Engine/Source/Editor/VirtualTexturingEditor/Classes/RuntimeVirtualTextureDetailsCustomization.h b/Engine/Source/Editor/VirtualTexturingEditor/Classes/RuntimeVirtualTextureDetailsCustomization.h
--- a/Engine/Source/Editor/VirtualTexturingEditor/Classes/RuntimeVirtualTextureDetailsCustomization.h
+++ b/Engine/Source/Editor/VirtualTexturingEditor/Classes/RuntimeVirtualTextureDetailsCustomization.h
@@ -67,6 +67,8 @@ protected:
 	bool IsBuildStreamedMipsEnabled() const;
 	/** Returns the visibility of the BuildAllStreamedMips button */
 	EVisibility GetBuildAllStreamedMipsVisible() const;
+	/** Returns true if the streaming texture is setup */
+	bool IsStreamingTextureSet() const;
 	/** Returns true if the streaming texture is up to date */
 	bool IsStreamingTextureValid(FText& OutReason) const;
 	/** Callback for Build button */
diff --git a/Engine/Source/Editor/VirtualTexturingEditor/Classes/VirtualTextureBuilderFactory.h b/Engine/Source/Editor/VirtualTexturingEditor/Classes/VirtualTextureBuilderFactory.h
--- a/Engine/Source/Editor/VirtualTexturingEditor/Classes/VirtualTextureBuilderFactory.h
+++ b/Engine/Source/Editor/VirtualTexturingEditor/Classes/VirtualTextureBuilderFactory.h
@@ -14,7 +14,9 @@ class UVirtualTextureBuilderFactory : public UFactory
 	GENERATED_UCLASS_BODY()
 
 	//~ Begin UFactory Interface
-	virtual bool ShouldShowInNewMenu() const { return false; }
 	virtual UObject* FactoryCreateNew(UClass* Class, UObject* InParent, FName Name, EObjectFlags Flags, UObject* Context, FFeedbackContext* Warn) override;
+	virtual FText GetDisplayName() const override;
+	virtual FText GetToolTip() const override;
+	virtual FString GetDefaultNewAssetName() const override;
 	//~ Begin UFactory Interface	
 };
diff --git a/Engine/Source/Editor/VirtualTexturingEditor/Private/RuntimeVirtualTextureDetailsCustomization.cpp b/Engine/Source/Editor/VirtualTexturingEditor/Private/RuntimeVirtualTextureDetailsCustomization.cpp
--- a/Engine/Source/Editor/VirtualTexturingEditor/Private/RuntimeVirtualTextureDetailsCustomization.cpp
+++ b/Engine/Source/Editor/VirtualTexturingEditor/Private/RuntimeVirtualTextureDetailsCustomization.cpp
@@ -345,7 +345,7 @@ void FRuntimeVirtualTextureComponentDetailsCustomization::CustomizeDetails(IDeta
 		.IsEnabled(this, &FRuntimeVirtualTextureComponentDetailsCustomization::IsSetBoundsEnabled)
 	];
 
-	static const FText BuildButtonSuffixText = LOCTEXT("Button_Build_Suffix_Tooltip", "Use the Build / Build Streaming Virtual Textures menu to build all versions of all Runtime Virtual Texture Volumes, after ensuring all actors rendering to RVT are loaded.");
+	static const FText BuildButtonSuffixText = LOCTEXT("Button_Build_Suffix_Tooltip", "Use the Build / Build Streaming Runtime Virtual Textures menu to build all versions of all Runtime Virtual Texture Volumes, after ensuring all actors rendering to RVT are loaded.");
 
 	// Apply custom widget for BuildStreamingMips.
 	TSharedRef<IPropertyHandle> BuildStreamingMipsPropertyHandle = DetailBuilder.GetProperty(TEXT("bBuildStreamingMipsButton"));
@@ -355,7 +355,7 @@ void FRuntimeVirtualTextureComponentDetailsCustomization::CustomizeDetails(IDeta
 		SNew(STextBlock)
 		.Font(IDetailLayoutBuilder::GetDetailFont())
 		.Text(LOCTEXT("Button_BuildStreamingTexture", "Build Streaming Texture"))
-		.ToolTipText(FText::Format(LOCTEXT("Button_Build_Prefix_Tooltip", "Build the low mips as streaming virtual texture data.\n{0}"), BuildButtonSuffixText))
+		.ToolTipText(FText::Format(LOCTEXT("Button_Build_Prefix_Tooltip", "Build runtime virtual texture low mips streaming data.\n{0}"), BuildButtonSuffixText))
 	]
 	.ValueContent()
 	.MinDesiredWidth(150.f)
@@ -369,7 +369,10 @@ void FRuntimeVirtualTextureComponentDetailsCustomization::CustomizeDetails(IDeta
 			.HAlign(HAlign_Center)
 			.VAlign(VAlign_Center)
 			.ContentPadding(2.f)
-			.Text(LOCTEXT("Button_Build", "Build"))
+			.Text_Lambda([this]()
+				{
+					return IsStreamingTextureSet() ? LOCTEXT("Button_Build", "Build") : LOCTEXT("Button_CreateAndBuild", "Create & Build");
+				})
 			.ToolTipText(FText::Format(LOCTEXT("Button_Build_Tooltip", "Build the low mips as streaming virtual texture data for the current shading path, using the currently loaded actors. \n"
 				"If \"Separate Texture For Mobile\" is enabled in the Streaming Texture, only the mobile version of the texture will be updated when using this button "
 				"while the mobile preview mode is active (and only the desktop version otherwise).\n{0}"), BuildButtonSuffixText))
@@ -409,7 +412,7 @@ void FRuntimeVirtualTextureComponentDetailsCustomization::CustomizeDetails(IDeta
 					const FText BuildAllShadingPathsText = LOCTEXT("BuildAllShadingPaths", " or Build All to rebuild it for all shading paths");
 					return FText::Format(LOCTEXT("Warning_Build_Tooltip", "The settings have changed since the Streaming Texture was last rebuilt.\n{0}\n"
 						"Use the Build button to rebuild the streaming virtual texture for the current shading path{1}.\n"
-						"Build Streaming Virtual Textures in the Build menu can also be used to rebuild all loaded streaming virtual textures for all shading paths, after ensuring all actors rendering to RVT are loaded.\n"
+						"Build Streaming Runtime Virtual Textures in the Build menu can also be used to rebuild all loaded streaming runtime virtual textures for all shading paths, after ensuring all actors rendering to RVT are loaded.\n"
 						"Meanwhile, streaming mips will be disabled."), 
 						Reason, (GetBuildAllStreamedMipsVisible() == EVisibility::Visible) ? BuildAllShadingPathsText : FText());
 				}
@@ -467,6 +470,12 @@ EVisibility FRuntimeVirtualTextureComponentDetailsCustomization::GetBuildAllStre
 	return EVisibility::Collapsed;
 }
 
+bool FRuntimeVirtualTextureComponentDetailsCustomization::IsStreamingTextureSet() const
+{
+	URuntimeVirtualTextureComponent* Component = RuntimeVirtualTextureComponent.Get();
+	return Component && (Component->GetStreamingTexture() != nullptr);
+}
+
 bool FRuntimeVirtualTextureComponentDetailsCustomization::IsStreamingTextureValid(FText& OutReason) const
 {
 	using EStreamingTextureStatusFlags = URuntimeVirtualTextureComponent::EStreamingTextureStatusFlags;
diff --git a/Engine/Source/Editor/VirtualTexturingEditor/Private/VirtualTextureBuilderAssetTypeActions.cpp b/Engine/Source/Editor/VirtualTexturingEditor/Private/VirtualTextureBuilderAssetTypeActions.cpp
--- a/Engine/Source/Editor/VirtualTexturingEditor/Private/VirtualTextureBuilderAssetTypeActions.cpp
+++ b/Engine/Source/Editor/VirtualTexturingEditor/Private/VirtualTextureBuilderAssetTypeActions.cpp
@@ -13,7 +13,7 @@ UClass* FAssetTypeActions_VirtualTextureBuilder::GetSupportedClass() const
 
 FText FAssetTypeActions_VirtualTextureBuilder::GetName() const
 {
-	return LOCTEXT("AssetTypeActions_VirtualTextureBuilder", "Virtual Texture Builder"); 
+	return LOCTEXT("AssetTypeActions_VirtualTextureBuilder", "Streaming Runtime Virtual Texture"); 
 }
 
 FColor FAssetTypeActions_VirtualTextureBuilder::GetTypeColor() const 
diff --git a/Engine/Source/Editor/VirtualTexturingEditor/Private/VirtualTextureBuilderFactory.cpp b/Engine/Source/Editor/VirtualTexturingEditor/Private/VirtualTextureBuilderFactory.cpp
--- a/Engine/Source/Editor/VirtualTexturingEditor/Private/VirtualTextureBuilderFactory.cpp
+++ b/Engine/Source/Editor/VirtualTexturingEditor/Private/VirtualTextureBuilderFactory.cpp
@@ -7,6 +7,8 @@
 
 #include UE_INLINE_GENERATED_CPP_BY_NAME(VirtualTextureBuilderFactory)
 
+#define LOCTEXT_NAMESPACE "VirtualTextureBuilderFactory"
+
 UVirtualTextureBuilderFactory::UVirtualTextureBuilderFactory(const FObjectInitializer& ObjectInitializer)
 	: Super(ObjectInitializer)
 {
@@ -21,3 +23,22 @@ UObject* UVirtualTextureBuilderFactory::FactoryCreateNew(UClass* Class, UObject*
 {
 	return NewObject<UVirtualTextureBuilder>(InParent, Class, Name, Flags);
 }
+
+FText UVirtualTextureBuilderFactory::GetDisplayName() const
+{
+	return FText(LOCTEXT("DisplayName", "Streaming Runtime Virtual Texture"));
+}
+
+FText UVirtualTextureBuilderFactory::GetToolTip() const
+{
+	return LOCTEXT("ToolTip", "The baked version of a runtime virtual texture (RVT) : this allows to capture the low mips of the RVT in order to save on rendering time, "
+		"so that the baked pages are streamed directly from disk instead of being rendered at runtime.\n"
+		"In order to work, this asset has to be associated with a Runtime Virtual Texture component. It also gets built through it.");
+}
+
+FString UVirtualTextureBuilderFactory::GetDefaultNewAssetName() const
+{
+	return FString(TEXT("NewStreamingRuntimeVirtualTexture"));
+}
+
+#undef LOCTEXT_NAMESPACE
\ No newline at end of file
diff --git a/Engine/Source/Runtime/Engine/Classes/Components/RuntimeVirtualTextureComponent.h b/Engine/Source/Runtime/Engine/Classes/Components/RuntimeVirtualTextureComponent.h
--- a/Engine/Source/Runtime/Engine/Classes/Components/RuntimeVirtualTextureComponent.h
+++ b/Engine/Source/Runtime/Engine/Classes/Components/RuntimeVirtualTextureComponent.h
@@ -93,7 +93,7 @@ protected:
 	bool bHidePrimitives = false;
 
 	/** Texture object containing streamed low mips. This can reduce rendering update cost. */
-	UPROPERTY(EditAnywhere, BlueprintReadOnly, Category = StreamingVirtualTexture)
+	UPROPERTY(EditAnywhere, BlueprintReadOnly, TextExportTransient, NonPIEDuplicateTransient, Category = StreamingVirtualTexture)
 	TObjectPtr<UVirtualTextureBuilder> StreamingTexture = nullptr;
 
 	/** Number of streaming low mips to build for the virtual texture. */
diff --git a/Engine/Source/Runtime/Engine/Classes/VT/VirtualTextureBuilder.h b/Engine/Source/Runtime/Engine/Classes/VT/VirtualTextureBuilder.h
--- a/Engine/Source/Runtime/Engine/Classes/VT/VirtualTextureBuilder.h
+++ b/Engine/Source/Runtime/Engine/Classes/VT/VirtualTextureBuilder.h
@@ -43,7 +43,7 @@ struct FVirtualTextureBuildDesc
  * This has a simple BuildTexture() interface but we may want to extend in the future to support partial builds
  * or other more blueprint driven approaches for data generation.
  */
-UCLASS(ClassGroup = Rendering, BlueprintType, MinimalAPI)
+UCLASS(ClassGroup = Rendering, BlueprintType, MinimalAPI, meta = (DisplayName = "Streaming Runtime Virtual Texture"))
 class UVirtualTextureBuilder : public UObject
 {
 public:
