Disable some checks in virtual texture conversion.
We now allow conversion when virtual textures feed opacity or vertex output since these are now valid scenarios (though they do need to be used carefully).

diff --git a/Engine/Source/Editor/VirtualTexturingEditor/Private/VirtualTextureConversionWorker.cpp b/Engine/Source/Editor/VirtualTexturingEditor/Private/VirtualTextureConversionWorker.cpp
--- a/Engine/Source/Editor/VirtualTexturingEditor/Private/VirtualTextureConversionWorker.cpp
+++ b/Engine/Source/Editor/VirtualTexturingEditor/Private/VirtualTextureConversionWorker.cpp
@@ -139,6 +139,7 @@ void FVirtualTextureConversionWorker::FindAllTexturesAndMaterials_Iteration(TSet
 
 		// Check all materials using texture, make sure we're able to convert this texture
 		bool bIsVirtualTextureValid = true;
+		/*
 		for (UMaterialInterface* Material : MaterialsUsingTexture)
 		{
 			if (!bConvertBackward && !IsVirtualTextureValidForMaterial(Material, Tex2D))
@@ -150,6 +151,7 @@ void FVirtualTextureConversionWorker::FindAllTexturesAndMaterials_Iteration(TSet
 				break;
 			}
 		}
+		*/
 
 		if (bIsVirtualTextureValid)
 		{
