Update the texture sample preview expressions when converting virtual textures.
If we don't update the expressions, then a stale expression will continue to use VT or non-VT sampling for the converted texture.
That leads to visual bugs for the preview.
#rnx

diff --git a/Engine/Source/Editor/VirtualTexturingEditor/Private/VirtualTextureConversionWorker.cpp b/Engine/Source/Editor/VirtualTexturingEditor/Private/VirtualTextureConversionWorker.cpp
--- a/Engine/Source/Editor/VirtualTexturingEditor/Private/VirtualTextureConversionWorker.cpp
+++ b/Engine/Source/Editor/VirtualTexturingEditor/Private/VirtualTextureConversionWorker.cpp
@@ -21,9 +21,11 @@
 #include "Editor.h"
 #include "Misc/ScopedSlowTask.h"
 #include "EditorSupportDelegates.h"
+#include "IMaterialEditor.h"
 #include "MaterialGraph/MaterialGraph.h"
 #include "MaterialEditingLibrary.h"
 #include "AssetRegistry/AssetRegistryHelpers.h"
+#include "Subsystems/AssetEditorSubsystem.h"
 #include "UObject/UObjectIterator.h"
 
 #define LOCTEXT_NAMESPACE "AssetVTConversion"
@@ -718,7 +720,18 @@ void FVirtualTextureConversionWorker::DoConvert()
 
 		UMaterialEditingLibrary::UpdateMaterialFunctions(ModifiedMaterialFunctions, UMaterialEditingLibrary::FOnItemComplete::CreateLambda([&MaterialTask]() { MaterialTask.EnterProgressFrame(); }));
 
-		// update the world's viewports
+		// Update expression previews for any open materials.
+		UAssetEditorSubsystem* AssetEditorSubsystem = GEditor->GetEditorSubsystem<UAssetEditorSubsystem>();
+		for (UMaterial* Mat : ModifiedMaterials)
+		{
+			if (IAssetEditorInstance* FoundAssetEditor = AssetEditorSubsystem->FindEditorForAsset(Mat, false))
+			{
+				IMaterialEditor* MaterialEditor = static_cast<IMaterialEditor*>(FoundAssetEditor);
+				MaterialEditor->ForceRefreshExpressionPreviews();
+			}
+		}
+
+		// Update the world's viewports.
 		UE_LOG(LogVirtualTextureConversion, Display, TEXT("Broadcasting to editor."));
 		FEditorDelegates::RefreshEditor.Broadcast();
 		FEditorSupportDelegates::RedrawAllViewports.Broadcast();
