[ContentBrowser] New Add Menu Texture Menu

#Synced-CL 52619962

#Synced-CL 52619962

diff --git a/Engine/Plugins/Editor/EngineAssetDefinitions/Source/Private/AssetDefinition_CanvasRenderTarget2D.h b/Engine/Plugins/Editor/EngineAssetDefinitions/Source/Private/AssetDefinition_CanvasRenderTarget2D.h
--- a/Engine/Plugins/Editor/EngineAssetDefinitions/Source/Private/AssetDefinition_CanvasRenderTarget2D.h
+++ b/Engine/Plugins/Editor/EngineAssetDefinitions/Source/Private/AssetDefinition_CanvasRenderTarget2D.h
@@ -16,5 +16,10 @@ public:
 	// UAssetDefinition Begin
 	virtual FText GetAssetDisplayName() const override { return NSLOCTEXT("AssetTypeActions", "AssetTypeActions_CanvasRenderTarget2D", "Canvas Render Target"); }
 	virtual TSoftClassPtr<UObject> GetAssetClass() const override { return UCanvasRenderTarget2D::StaticClass(); }
+	virtual TConstArrayView<FAssetCategoryPath> GetAssetCategories() const override
+	{
+		static const auto Categories = { FAssetCategoryPath(EAssetCategoryPaths::Texture) };
+		return Categories;
+	}
 	// UAssetDefinition End
 };
diff --git a/Engine/Plugins/Editor/EngineAssetDefinitions/Source/Private/AssetDefinition_Texture2DArray.h b/Engine/Plugins/Editor/EngineAssetDefinitions/Source/Private/AssetDefinition_Texture2DArray.h
--- a/Engine/Plugins/Editor/EngineAssetDefinitions/Source/Private/AssetDefinition_Texture2DArray.h
+++ b/Engine/Plugins/Editor/EngineAssetDefinitions/Source/Private/AssetDefinition_Texture2DArray.h
@@ -17,5 +17,10 @@ public:
 	virtual FText GetAssetDisplayName() const override { return NSLOCTEXT("AssetTypeActions", "AssetTypeActions_Texture2DArray", "Texture 2D Array"); }
 	virtual FLinearColor GetAssetColor() const override { return FLinearColor(FColor(0, 64, 128)); }
 	virtual TSoftClassPtr<UObject> GetAssetClass() const override { return UTexture2DArray::StaticClass(); }
+	virtual TConstArrayView<FAssetCategoryPath> GetAssetCategories() const override
+	{
+		static const auto Categories = { FAssetCategoryPath(EAssetCategoryPaths::Texture) };
+		return Categories;
+	}
 	// UAssetDefinition End
 };
diff --git a/Engine/Plugins/Editor/EngineAssetDefinitions/Source/Private/AssetDefinition_TextureCubeArray.h b/Engine/Plugins/Editor/EngineAssetDefinitions/Source/Private/AssetDefinition_TextureCubeArray.h
--- a/Engine/Plugins/Editor/EngineAssetDefinitions/Source/Private/AssetDefinition_TextureCubeArray.h
+++ b/Engine/Plugins/Editor/EngineAssetDefinitions/Source/Private/AssetDefinition_TextureCubeArray.h
@@ -17,5 +17,10 @@ public:
 	virtual FText GetAssetDisplayName() const override { return NSLOCTEXT("AssetTypeActions", "AssetTypeActions_TextureCubeArray", "Texture Cube Array"); }
 	virtual FLinearColor GetAssetColor() const override { return FLinearColor(FColor(192, 192, 128)); }
 	virtual TSoftClassPtr<UObject> GetAssetClass() const override { return UTextureCubeArray::StaticClass(); }
+	virtual TConstArrayView<FAssetCategoryPath> GetAssetCategories() const override
+	{
+		static const auto Categories = { FAssetCategoryPath(EAssetCategoryPaths::Texture, NSLOCTEXT("AssetDefinition", "TextureCubeArray_SubMenu", "Texture Array"), ECategoryMenuType::Section) };
+		return Categories;
+	}
 	// UAssetDefinition End
 };
diff --git a/Engine/Plugins/Editor/EngineAssetDefinitions/Source/Private/AssetDefinition_TextureRenderTarget2D.h b/Engine/Plugins/Editor/EngineAssetDefinitions/Source/Private/AssetDefinition_TextureRenderTarget2D.h
--- a/Engine/Plugins/Editor/EngineAssetDefinitions/Source/Private/AssetDefinition_TextureRenderTarget2D.h
+++ b/Engine/Plugins/Editor/EngineAssetDefinitions/Source/Private/AssetDefinition_TextureRenderTarget2D.h
@@ -16,5 +16,10 @@ public:
 	// UAssetDefinition Begin
 	virtual FText GetAssetDisplayName() const override { return NSLOCTEXT("AssetTypeActions", "AssetTypeActions_TextureRenderTarget2D", "Render Target"); }
 	virtual TSoftClassPtr<UObject> GetAssetClass() const override { return UTextureRenderTarget2D::StaticClass(); }
+	virtual TConstArrayView<FAssetCategoryPath> GetAssetCategories() const override
+	{
+		static const auto Categories = { FAssetCategoryPath(EAssetCategoryPaths::Texture) };
+		return Categories;
+	}
 	// UAssetDefinition End
 };
diff --git a/Engine/Plugins/Editor/EngineAssetDefinitions/Source/Private/AssetDefinition_TextureRenderTarget2DArray.h b/Engine/Plugins/Editor/EngineAssetDefinitions/Source/Private/AssetDefinition_TextureRenderTarget2DArray.h
--- a/Engine/Plugins/Editor/EngineAssetDefinitions/Source/Private/AssetDefinition_TextureRenderTarget2DArray.h
+++ b/Engine/Plugins/Editor/EngineAssetDefinitions/Source/Private/AssetDefinition_TextureRenderTarget2DArray.h
@@ -16,5 +16,10 @@ public:
 	// UAssetDefinition Begin
 	virtual FText GetAssetDisplayName() const override { return NSLOCTEXT("AssetTypeActions", "AssetTypeActions_TextureRenderTarget2DArray", "2D Array Render Target"); }
 	virtual TSoftClassPtr<UObject> GetAssetClass() const override { return UTextureRenderTarget2DArray::StaticClass(); }
+	virtual TConstArrayView<FAssetCategoryPath> GetAssetCategories() const override
+	{
+		static const auto Categories = { FAssetCategoryPath(EAssetCategoryPaths::Texture, NSLOCTEXT("AssetDefinition", "TextureRenderTarget2DArray_SubMenu", "Render Target"), ECategoryMenuType::Section) };
+		return Categories;
+	}
 	// UAssetDefinition End
 };
diff --git a/Engine/Plugins/Editor/EngineAssetDefinitions/Source/Private/AssetDefinition_TextureRenderTargetCube.h b/Engine/Plugins/Editor/EngineAssetDefinitions/Source/Private/AssetDefinition_TextureRenderTargetCube.h
--- a/Engine/Plugins/Editor/EngineAssetDefinitions/Source/Private/AssetDefinition_TextureRenderTargetCube.h
+++ b/Engine/Plugins/Editor/EngineAssetDefinitions/Source/Private/AssetDefinition_TextureRenderTargetCube.h
@@ -16,5 +16,10 @@ public:
 	// UAssetDefinition Begin
 	virtual FText GetAssetDisplayName() const override { return NSLOCTEXT("AssetTypeActions", "AssetTypeActions_TextureRenderTargetCube", "Cube Render Target"); }
 	virtual TSoftClassPtr<UObject> GetAssetClass() const override { return UTextureRenderTargetCube::StaticClass(); }
+	virtual TConstArrayView<FAssetCategoryPath> GetAssetCategories() const override
+	{
+		static const auto Categories = { FAssetCategoryPath(EAssetCategoryPaths::Texture, NSLOCTEXT("AssetDefinition", "TextureRenderTargetCube_SubMenu", "Render Target"), ECategoryMenuType::Section) };
+		return Categories;
+	}
 	// UAssetDefinition End
 };
diff --git a/Engine/Plugins/Editor/EngineAssetDefinitions/Source/Private/AssetDefinition_TextureRenderTargetVolume.h b/Engine/Plugins/Editor/EngineAssetDefinitions/Source/Private/AssetDefinition_TextureRenderTargetVolume.h
--- a/Engine/Plugins/Editor/EngineAssetDefinitions/Source/Private/AssetDefinition_TextureRenderTargetVolume.h
+++ b/Engine/Plugins/Editor/EngineAssetDefinitions/Source/Private/AssetDefinition_TextureRenderTargetVolume.h
@@ -16,5 +16,10 @@ public:
 	// UAssetDefinition Begin
 	virtual FText GetAssetDisplayName() const override { return NSLOCTEXT("AssetTypeActions", "AssetTypeActions_TextureRenderTargetVolume", "Volume Render Target"); }
 	virtual TSoftClassPtr<UObject> GetAssetClass() const override { return UTextureRenderTargetVolume::StaticClass(); }
+	virtual TConstArrayView<FAssetCategoryPath> GetAssetCategories() const override
+	{
+		static const auto Categories = { FAssetCategoryPath(EAssetCategoryPaths::Texture, NSLOCTEXT("AssetDefinition", "TextureRenderTargetVolume_SubMenu", "Render Target"), ECategoryMenuType::Section) };
+		return Categories;
+	}
 	// UAssetDefinition End
 };
diff --git a/Engine/Plugins/Editor/EngineAssetDefinitions/Source/Private/AssetDefinition_VolumeTexture.h b/Engine/Plugins/Editor/EngineAssetDefinitions/Source/Private/AssetDefinition_VolumeTexture.h
--- a/Engine/Plugins/Editor/EngineAssetDefinitions/Source/Private/AssetDefinition_VolumeTexture.h
+++ b/Engine/Plugins/Editor/EngineAssetDefinitions/Source/Private/AssetDefinition_VolumeTexture.h
@@ -17,5 +17,10 @@ public:
 	virtual FText GetAssetDisplayName() const override { return NSLOCTEXT("AssetTypeActions", "AssetTypeActions_VolumeTexture", "Volume Texture"); }
 	virtual FLinearColor GetAssetColor() const override { return FLinearColor(FColor(192, 128, 64)); }
 	virtual TSoftClassPtr<UObject> GetAssetClass() const override { return UVolumeTexture::StaticClass(); }
+	virtual TConstArrayView<FAssetCategoryPath> GetAssetCategories() const override
+	{
+		static const auto Categories = { FAssetCategoryPath(EAssetCategoryPaths::Texture) };
+		return Categories;
+	}
 	// UAssetDefinition End
 };
diff --git a/Engine/Plugins/Media/PixelStreaming2/Source/PixelStreaming2Editor/PixelStreaming2Editor.build.cs b/Engine/Plugins/Media/PixelStreaming2/Source/PixelStreaming2Editor/PixelStreaming2Editor.build.cs
--- a/Engine/Plugins/Media/PixelStreaming2/Source/PixelStreaming2Editor/PixelStreaming2Editor.build.cs
+++ b/Engine/Plugins/Media/PixelStreaming2/Source/PixelStreaming2Editor/PixelStreaming2Editor.build.cs
@@ -14,6 +14,7 @@ namespace UnrealBuildTool.Rules
 
 			PrivateDependencyModuleNames.AddRange(new string[]
 			{
+				"AssetDefinition",
 				"Core",
 				"CoreUObject",
 				"Engine",
diff --git a/Engine/Plugins/Media/PixelStreaming2/Source/PixelStreaming2Editor/Private/AssetDefinition_PixelStreaming2MediaTexture.h b/Engine/Plugins/Media/PixelStreaming2/Source/PixelStreaming2Editor/Private/AssetDefinition_PixelStreaming2MediaTexture.h
new file mode 100644
--- /dev/null
+++ b/Engine/Plugins/Media/PixelStreaming2/Source/PixelStreaming2Editor/Private/AssetDefinition_PixelStreaming2MediaTexture.h
@@ -0,0 +1,26 @@
+// Copyright Epic Games, Inc. All Rights Reserved.
+
+#pragma once
+#include "AssetDefinitionDefault.h"
+#include "Blueprints/PixelStreaming2MediaTexture.h"
+#include "AssetDefinition_PixelStreaming2MediaTexture.generated.h"
+
+UCLASS()
+class UAssetDefinition_PixelStreaming2MediaTexture : public UAssetDefinitionDefault
+{
+	GENERATED_BODY()
+public:
+	// UAssetDefinition Begin
+	virtual FText GetAssetDisplayName() const override { return NSLOCTEXT("AssetDefinition", "AssetDefinition_PixelStreaming2MediaTexture", "Pixel Streaming 2 Media Texture"); }
+	virtual TSoftClassPtr<UObject> GetAssetClass() const override { return UPixelStreaming2MediaTexture::StaticClass(); }
+	virtual FLinearColor GetAssetColor() const override { return FColor(192, 64, 64); }
+	virtual TConstArrayView<FAssetCategoryPath> GetAssetCategories() const override
+	{
+		static const auto Categories = 
+		{
+			FAssetCategoryPath(EAssetCategoryPaths::Texture, NSLOCTEXT("AssetDefinition", "AssetDefinition_PixelStreaming2MediaTextureSubMenu", "Other"), ECategoryMenuType::Section)
+		};
+		return Categories;
+	}
+	// UAssetDefinition End
+};
diff --git a/Engine/Plugins/Media/PixelStreaming2/Source/PixelStreaming2Editor/Private/Factories.cpp b/Engine/Plugins/Media/PixelStreaming2/Source/PixelStreaming2Editor/Private/Factories.cpp
--- a/Engine/Plugins/Media/PixelStreaming2/Source/PixelStreaming2Editor/Private/Factories.cpp
+++ b/Engine/Plugins/Media/PixelStreaming2/Source/PixelStreaming2Editor/Private/Factories.cpp
@@ -29,11 +29,6 @@ FText UPixelStreaming2MediaTextureFactory::GetDisplayName() const
 	return LOCTEXT("MediaTextureFactoryDisplayName", "Pixel Streaming 2 Media Texture");
 }
 
-uint32 UPixelStreaming2MediaTextureFactory::GetMenuCategories() const
-{
-	return EAssetTypeCategories::Textures;
-}
-
 UObject* UPixelStreaming2MediaTextureFactory::FactoryCreateNew(UClass* InClass, UObject* InParent, FName InName, EObjectFlags Flags, UObject* Context, FFeedbackContext* Warn)
 {
 	if (UPixelStreaming2MediaTexture* Resource = NewObject<UPixelStreaming2MediaTexture>(InParent, InName, Flags | RF_Transactional))
diff --git a/Engine/Plugins/Media/PixelStreaming2/Source/PixelStreaming2Editor/Private/Factories.h b/Engine/Plugins/Media/PixelStreaming2/Source/PixelStreaming2Editor/Private/Factories.h
--- a/Engine/Plugins/Media/PixelStreaming2/Source/PixelStreaming2Editor/Private/Factories.h
+++ b/Engine/Plugins/Media/PixelStreaming2/Source/PixelStreaming2Editor/Private/Factories.h
@@ -16,7 +16,6 @@ class UPixelStreaming2MediaTextureFactory : public UFactory
 
 public:
 	virtual FText	 GetDisplayName() const override;
-	virtual uint32	 GetMenuCategories() const override;
 	virtual UObject* FactoryCreateNew(UClass* InClass, UObject* InParent, FName InName, EObjectFlags Flags, UObject* Context, FFeedbackContext* Warn) override;
 };
 
diff --git a/Engine/Plugins/TextureGraph/Source/TextureGraphEditor/Private/AssetDefinition_TextureGraph.h b/Engine/Plugins/TextureGraph/Source/TextureGraphEditor/Private/AssetDefinition_TextureGraph.h
--- a/Engine/Plugins/TextureGraph/Source/TextureGraphEditor/Private/AssetDefinition_TextureGraph.h
+++ b/Engine/Plugins/TextureGraph/Source/TextureGraphEditor/Private/AssetDefinition_TextureGraph.h
@@ -19,7 +19,7 @@ public:
 	virtual TSoftClassPtr<UObject> GetAssetClass() const override { return UTextureGraph::StaticClass(); }
 	virtual TConstArrayView<FAssetCategoryPath> GetAssetCategories() const override
 	{
-		static const auto Categories = { EAssetCategoryPaths::Texture };
+		static const auto Categories = { FAssetCategoryPath(EAssetCategoryPaths::Texture, NSLOCTEXT("AssetDefinition", "TextureGraph_SubMenu", "Texture Graph"), ECategoryMenuType::Section) };
 		return Categories;
 	}
 
diff --git a/Engine/Plugins/TextureGraph/Source/TextureGraphEditor/Private/AssetDefinition_TextureGraphInstance.h b/Engine/Plugins/TextureGraph/Source/TextureGraphEditor/Private/AssetDefinition_TextureGraphInstance.h
--- a/Engine/Plugins/TextureGraph/Source/TextureGraphEditor/Private/AssetDefinition_TextureGraphInstance.h
+++ b/Engine/Plugins/TextureGraph/Source/TextureGraphEditor/Private/AssetDefinition_TextureGraphInstance.h
@@ -19,7 +19,7 @@ public:
 	virtual TSoftClassPtr<UObject> GetAssetClass() const override { return UTextureGraphInstance::StaticClass(); }
 	virtual TConstArrayView<FAssetCategoryPath> GetAssetCategories() const override
 	{
-		static const auto Categories = { EAssetCategoryPaths::Texture };
+		static const auto Categories = { FAssetCategoryPath(EAssetCategoryPaths::Texture, NSLOCTEXT("AssetDefinition", "TextureGraphInstance_SubMenu", "Texture Graph"), ECategoryMenuType::Section) };
 		return Categories;
 	}
 
diff --git a/Engine/Plugins/VirtualProduction/DynamicMaterial/Source/DynamicMaterialTextureSetEditor/Private/AssetDefinition_DMTextureSet.cpp b/Engine/Plugins/VirtualProduction/DynamicMaterial/Source/DynamicMaterialTextureSetEditor/Private/AssetDefinition_DMTextureSet.cpp
--- a/Engine/Plugins/VirtualProduction/DynamicMaterial/Source/DynamicMaterialTextureSetEditor/Private/AssetDefinition_DMTextureSet.cpp
+++ b/Engine/Plugins/VirtualProduction/DynamicMaterial/Source/DynamicMaterialTextureSetEditor/Private/AssetDefinition_DMTextureSet.cpp
@@ -30,7 +30,10 @@ FLinearColor UAssetDefinition_DMTextureSet::GetAssetColor() const
 
 TConstArrayView<FAssetCategoryPath> UAssetDefinition_DMTextureSet::GetAssetCategories() const
 {
-	static const TArray<FAssetCategoryPath> Categories = {EAssetCategoryPaths::Texture};
+	static const TArray<FAssetCategoryPath> Categories = 
+		{
+			FAssetCategoryPath(EAssetCategoryPaths::Texture, LOCTEXT("AssetDefinition_DMTextureSetSubMenu", "Other"), ECategoryMenuType::Section)
+		};
 	return Categories;
 }
 
diff --git a/Engine/Source/Editor/VirtualTexturingEditor/Classes/MaterialCacheVirtualTextureTagAssetTypeActions.h b/Engine/Source/Editor/VirtualTexturingEditor/Classes/MaterialCacheVirtualTextureTagAssetTypeActions.h
deleted file mode 100644
diff --git a/Engine/Source/Editor/VirtualTexturingEditor/Classes/RuntimeVirtualTextureAssetTypeActions.h b/Engine/Source/Editor/VirtualTexturingEditor/Classes/RuntimeVirtualTextureAssetTypeActions.h
deleted file mode 100644
diff --git a/Engine/Source/Editor/VirtualTexturingEditor/Classes/VirtualTextureBuilderAssetTypeActions.h b/Engine/Source/Editor/VirtualTexturingEditor/Classes/VirtualTextureBuilderAssetTypeActions.h
deleted file mode 100644
diff --git a/Engine/Source/Editor/VirtualTexturingEditor/Private/AssetDefinition_MaterialCacheVirtualTextureTag.h b/Engine/Source/Editor/VirtualTexturingEditor/Private/AssetDefinition_MaterialCacheVirtualTextureTag.h
new file mode 100644
--- /dev/null
+++ b/Engine/Source/Editor/VirtualTexturingEditor/Private/AssetDefinition_MaterialCacheVirtualTextureTag.h
@@ -0,0 +1,26 @@
+// Copyright Epic Games, Inc. All Rights Reserved.
+
+#pragma once
+#include "AssetDefinitionDefault.h"
+#include "MaterialCache/MaterialCacheVirtualTextureTag.h"
+#include "AssetDefinition_MaterialCacheVirtualTextureTag.generated.h"
+
+UCLASS()
+class UAssetDefinition_MaterialCacheVirtualTextureTag : public UAssetDefinitionDefault
+{
+	GENERATED_BODY()
+public:
+	// UAssetDefinition Begin
+	virtual FText GetAssetDisplayName() const override { return NSLOCTEXT("AssetDefinition", "AssetDefinition_MaterialCacheVirtualTextureTag", "Material Cache Virtual Texture Tag"); }
+	virtual TSoftClassPtr<UObject> GetAssetClass() const override { return UMaterialCacheVirtualTextureTag::StaticClass(); }
+	virtual FLinearColor GetAssetColor() const override { return FColor(255, 100, 100); }
+	virtual TConstArrayView<FAssetCategoryPath> GetAssetCategories() const override
+	{
+		static const auto Categories = 
+		{
+			FAssetCategoryPath(EAssetCategoryPaths::Texture, NSLOCTEXT("AssetDefinition", "AssetDefinition_MaterialCacheVirtualTextureTagSubMenu", "Virtual Texture"), ECategoryMenuType::Section)
+		};
+		return Categories;
+	}
+	// UAssetDefinition End
+};
diff --git a/Engine/Source/Editor/VirtualTexturingEditor/Private/AssetDefinition_RuntimeVirtualTexture.cpp b/Engine/Source/Editor/VirtualTexturingEditor/Private/AssetDefinition_RuntimeVirtualTexture.cpp
new file mode 100644
--- /dev/null
+++ b/Engine/Source/Editor/VirtualTexturingEditor/Private/AssetDefinition_RuntimeVirtualTexture.cpp
@@ -0,0 +1,281 @@
+// Copyright Epic Games, Inc. All Rights Reserved.
+
+#include "AssetDefinition_RuntimeVirtualTexture.h"
+#include "AssetRegistry/AssetRegistryHelpers.h"
+#include "ContentBrowserMenuContexts.h"
+#include "ContentBrowserModule.h"
+#include "Editor.h"
+#include "EditorSupportDelegates.h"
+#include "FileHelpers.h"
+#include "IContentBrowserSingleton.h"
+#include "MaterialEditingLibrary.h"
+#include "Materials/Material.h"
+#include "Materials/MaterialExpressionRuntimeVirtualTextureSample.h"
+#include "Materials/MaterialExpressionTextureBase.h"
+#include "Materials/MaterialFunction.h"
+#include "Materials/MaterialInstance.h"
+#include "Materials/MaterialInterface.h"
+#include "ToolMenus.h"
+
+#define LOCTEXT_NAMESPACE "AssetDefinition_RuntimeVirtualTexture"
+
+DEFINE_LOG_CATEGORY_STATIC(LogRuntimeVirtualTextureFixMaterial, Log, Log);
+
+namespace MenuExtension::RuntimeVirtualTexture
+{
+	template <class T> 
+	void GetReferencersOfType(UObject *Object, TArray<T*> &OutObjects)
+	{
+		TArray<FAssetData> AssetDatas;
+		UAssetRegistryHelpers::FindReferencersOfAssetOfClass(Object, { T::StaticClass() }, AssetDatas);
+
+		for (auto Data : AssetDatas)
+		{
+			T* TypedAsset = Cast<T>(Data.GetAsset());
+			if (TypedAsset != nullptr)
+			{
+				OutObjects.Add(TypedAsset);
+			}
+		}
+	}
+
+	void FindAllMaterials(URuntimeVirtualTexture* RuntimeVirtualTexture, TArray<UMaterial*>& OutMaterials, TArray<UMaterialFunctionInterface*>& OutFunctions)
+	{
+		TArray<UMaterialInterface*> MaterialInterfaces;
+		GetReferencersOfType(RuntimeVirtualTexture, MaterialInterfaces);
+
+		for (UMaterialInterface *MaterialInterface : MaterialInterfaces)
+		{
+			UMaterial* Material = Cast<UMaterial>(MaterialInterface);
+			
+			UMaterialInstance* MaterialInstance = Cast<UMaterialInstance>(MaterialInterface);
+			if (MaterialInstance != nullptr)
+			{
+				Material = MaterialInstance->GetMaterial();
+			}
+
+			OutMaterials.AddUnique(Material);
+		}
+
+		GetReferencersOfType(RuntimeVirtualTexture, OutFunctions);
+	}
+
+	void FixMaterialUsage(URuntimeVirtualTexture* RuntimeVirtualTexture)
+	{
+		TArray<UPackage*> PackagesToSave;
+		{
+			UE_LOG(LogRuntimeVirtualTextureFixMaterial, Log, TEXT("Begin fix material usage for '%s' ..."), *RuntimeVirtualTexture->GetName());
+
+			TArray<UMaterial*> Materials;
+			TArray<UMaterialFunctionInterface*> Functions;
+			FindAllMaterials(RuntimeVirtualTexture, Materials, Functions);
+
+			int32 TaskCount = Materials.Num() + Functions.Num();
+			FScopedSlowTask Task(static_cast<float>(TaskCount), LOCTEXT("RuntimeVirtualTexture_FixMaterialUsageProgress", "Fixing materials for Runtime Virtual Texture usage..."));
+			Task.MakeDialog();
+
+			for (UMaterial* Material : Materials)
+			{
+				Task.EnterProgressFrame();
+
+				bool bMaterialModified = false;
+				for (UMaterialExpression* Expression : Material->GetExpressions())
+				{
+					UMaterialExpressionRuntimeVirtualTextureSample* RVTSampleExpression = Cast<UMaterialExpressionRuntimeVirtualTextureSample>(Expression);
+					if (RVTSampleExpression)
+					{
+						if (RuntimeVirtualTexture == RVTSampleExpression->VirtualTexture)
+						{
+							if (RVTSampleExpression->InitVirtualTextureDependentSettings())
+							{
+								Expression->Modify();
+
+								FPropertyChangedEvent Event(UMaterialExpressionTextureBase::StaticClass()->FindPropertyByName(GET_MEMBER_NAME_CHECKED(UMaterialExpressionRuntimeVirtualTextureSample, MaterialType)));
+								Expression->PostEditChangeProperty(Event);
+
+								bMaterialModified = true;
+							}
+						}
+					}
+				}
+
+				if (bMaterialModified)
+				{
+					UE_LOG(LogRuntimeVirtualTextureFixMaterial, Log, TEXT("  Recompile material '%s' ..."), *Material->GetName());
+
+					FScopedSlowTask CompileTask(1, FText::AsCultureInvariant(Material->GetName()));
+					CompileTask.MakeDialog();
+					CompileTask.EnterProgressFrame();
+
+					UMaterialEditingLibrary::RecompileMaterial(Material);
+
+					PackagesToSave.Add(Material->GetOutermost());
+				}
+			}
+
+			for (UMaterialFunctionInterface *Function : Functions)
+			{
+				Task.EnterProgressFrame();
+
+				bool bFunctionModified = false;
+				for (const TObjectPtr<UMaterialExpression>& Expression : Function->GetExpressions())
+				{
+					UMaterialExpressionRuntimeVirtualTextureSample* RVTSampleExpression = Cast<UMaterialExpressionRuntimeVirtualTextureSample>(Expression);
+					if (RVTSampleExpression)
+					{
+						if (RuntimeVirtualTexture == RVTSampleExpression->VirtualTexture)
+						{
+							if (RVTSampleExpression->InitVirtualTextureDependentSettings())
+							{
+								Expression->Modify();
+
+								FPropertyChangedEvent Event(UMaterialExpressionTextureBase::StaticClass()->FindPropertyByName(GET_MEMBER_NAME_CHECKED(UMaterialExpressionRuntimeVirtualTextureSample, MaterialType)));
+								Expression->PostEditChangeProperty(Event);
+
+								bFunctionModified = true;
+							}
+						}
+					}
+				}
+
+				if (bFunctionModified)
+				{
+					UE_LOG(LogRuntimeVirtualTextureFixMaterial, Log, TEXT("  Update function '%s' ..."), *Function->GetName());
+
+					FScopedSlowTask CompileTask(1, FText::AsCultureInvariant(Function->GetName()));
+					CompileTask.MakeDialog();
+					CompileTask.EnterProgressFrame();
+
+					UMaterialEditingLibrary::UpdateMaterialFunction(Function, nullptr);
+
+					PackagesToSave.Add(Function->GetOutermost());
+				}
+			}
+
+			UE_LOG(LogRuntimeVirtualTextureFixMaterial, Log, TEXT("End fix material usage for '%s' ..."), *RuntimeVirtualTexture->GetName());
+		}
+	
+		if (PackagesToSave.Num())
+		{
+			FEditorFileUtils::PromptForCheckoutAndSave(PackagesToSave, false, true);
+		}
+	}
+
+	void ExecuteFindMaterials(const FToolMenuContext& InContext)
+	{
+		const UContentBrowserAssetContextMenuContext* CBContext = UContentBrowserAssetContextMenuContext::FindContextWithAssets(InContext);
+		if (CBContext == nullptr)
+		{
+			return;
+		}
+
+		TArray<FAssetData> AssetData = CBContext->SelectedAssets;
+		if (AssetData.Num() != 1)
+		{
+			return;
+		}
+
+		TArray<FAssetData> Materials;
+		URuntimeVirtualTexture* RuntimeVirtualTexture = CBContext->LoadFirstSelectedObject<URuntimeVirtualTexture>();
+
+		if (RuntimeVirtualTexture != nullptr)
+		{
+			UAssetRegistryHelpers::FindReferencersOfAssetOfClass(RuntimeVirtualTexture, { UMaterialInterface::StaticClass(), UMaterialFunction::StaticClass() }, Materials);
+		}
+
+		if (Materials.Num() > 0)
+		{
+			FContentBrowserModule& ContentBrowserModule = FModuleManager::LoadModuleChecked<FContentBrowserModule>("ContentBrowser");
+			ContentBrowserModule.Get().SyncBrowserToAssets(Materials);
+		}
+	}
+
+	void ExecuteFixMaterialUsage(const FToolMenuContext& InContext)
+	{
+		const UContentBrowserAssetContextMenuContext* CBContext = UContentBrowserAssetContextMenuContext::FindContextWithAssets(InContext);
+		if (CBContext == nullptr)
+		{
+			return;
+		}
+
+		TArray<FAssetData> AssetData = CBContext->SelectedAssets;
+		if (AssetData.Num() != 1)
+		{
+			return;
+		}
+
+		URuntimeVirtualTexture* RuntimeVirtualTexture = CBContext->LoadFirstSelectedObject<URuntimeVirtualTexture>();
+		if (RuntimeVirtualTexture != nullptr)
+		{
+			FixMaterialUsage(RuntimeVirtualTexture);
+
+			FEditorDelegates::RefreshEditor.Broadcast();
+			FEditorSupportDelegates::RedrawAllViewports.Broadcast();
+		}
+	}
+
+	bool IsOneAssetSelected(const FToolMenuContext& InContext)
+	{
+		const UContentBrowserAssetContextMenuContext* CBContext = UContentBrowserAssetContextMenuContext::FindContextWithAssets(InContext);
+		if (CBContext == nullptr)
+		{
+			return false;
+		}
+
+		TArray<FAssetData> AssetData = CBContext->SelectedAssets;
+		return AssetData.Num() == 1;
+	}
+
+	static FDelayedAutoRegisterHelper DelayedAutoRegister(EDelayedRegisterRunPhase::EndOfEngineInit, []
+	{
+		UToolMenus::RegisterStartupCallback(FSimpleMulticastDelegate::FDelegate::CreateLambda([]()
+		{
+			FToolMenuOwnerScoped OwnerScoped(UE_MODULE_NAME);
+			
+			// Runtime Virtual Texture Action Registration
+			{
+				UToolMenu* Menu = UE::ContentBrowser::ExtendToolMenu_AssetContextMenu(URuntimeVirtualTexture::StaticClass());
+
+				FToolMenuSection& RuntimeVirtualTextureSection = Menu->FindOrAddSection("GetAssetActions");
+				RuntimeVirtualTextureSection.AddDynamicEntry(NAME_None, FNewToolMenuSectionDelegate::CreateLambda([](FToolMenuSection& InSection)
+					{
+						if (const UContentBrowserAssetContextMenuContext* Context = UContentBrowserAssetContextMenuContext::FindContextWithAssets(InSection))
+						{
+							if (Context->SelectedAssets.Num() > 0)
+							{
+								const FAssetData& AssetData = Context->SelectedAssets[0];
+								if (AssetData.AssetClassPath == URuntimeVirtualTexture::StaticClass()->GetClassPathName())
+								{
+									// Find Material
+									{
+										const TAttribute<FText> Label = LOCTEXT("RuntimeVirtualTexture_FindMaterials", "Find Materials Using This");
+										const TAttribute<FText> ToolTip = LOCTEXT("RuntimeVirtualTexture_FindMaterialsTooltip", "Finds all materials that use this texture in the content browser.");
+										const FSlateIcon Icon = FSlateIcon();
+
+										FToolUIAction UIAction;
+										UIAction.ExecuteAction = FToolMenuExecuteAction::CreateStatic(&ExecuteFindMaterials);
+										UIAction.IsActionVisibleDelegate = FToolMenuIsActionButtonVisible::CreateStatic(&IsOneAssetSelected);
+										InSection.AddMenuEntry("DialogWave_PlaySound", Label, ToolTip, Icon, UIAction);
+									}
+
+									// Fix Material
+									{
+										const TAttribute<FText> Label = LOCTEXT("RuntimeVirtualTexture_FixMaterialUsage", "Fix Material Usage");
+										const TAttribute<FText> ToolTip = LOCTEXT("RuntimeVirtualTexture_FixMaterialUsageTooltip", "Find materials using this Runtime Virtual Texture and fix any mismatching content types.");
+										const FSlateIcon Icon = FSlateIcon(FAppStyle::GetAppStyleSetName(), "ClassIcon.Texture2D");
+
+										FToolUIAction UIAction;
+										UIAction.ExecuteAction = FToolMenuExecuteAction::CreateStatic(&ExecuteFixMaterialUsage);
+										UIAction.IsActionVisibleDelegate = FToolMenuIsActionButtonVisible::CreateStatic(&IsOneAssetSelected);
+										InSection.AddMenuEntry("DialogWave_StopSound", Label, ToolTip, Icon, UIAction);
+									}
+								}
+							}
+						}
+					}));
+			}
+		}));
+	});
+}
+
+#undef LOCTEXT_NAMESPACE
diff --git a/Engine/Source/Editor/VirtualTexturingEditor/Private/AssetDefinition_RuntimeVirtualTexture.h b/Engine/Source/Editor/VirtualTexturingEditor/Private/AssetDefinition_RuntimeVirtualTexture.h
new file mode 100644
--- /dev/null
+++ b/Engine/Source/Editor/VirtualTexturingEditor/Private/AssetDefinition_RuntimeVirtualTexture.h
@@ -0,0 +1,23 @@
+// Copyright Epic Games, Inc. All Rights Reserved.
+
+#pragma once
+#include "AssetDefinitionDefault.h"
+#include "VT/RuntimeVirtualTexture.h"
+#include "AssetDefinition_RuntimeVirtualTexture.generated.h"
+
+UCLASS()
+class UAssetDefinition_RuntimeVirtualTexture : public UAssetDefinitionDefault
+{
+	GENERATED_BODY()
+public:
+	// UAssetDefinition Begin
+	virtual FText GetAssetDisplayName() const override { return NSLOCTEXT("AssetDefinition", "AssetDefinition_RuntimeVirtualTexture", "Runtime Virtual Texture"); }
+	virtual FLinearColor GetAssetColor() const override { return FColor(128, 128, 128); }
+	virtual TSoftClassPtr<UObject> GetAssetClass() const override { return URuntimeVirtualTexture::StaticClass(); }
+	virtual TConstArrayView<FAssetCategoryPath> GetAssetCategories() const override
+	{
+		static const auto Categories = { FAssetCategoryPath(EAssetCategoryPaths::Texture, NSLOCTEXT("AssetDefinition", "AssetDefinition_RuntimeVirtualTextureSubMenu", "Virtual Texture"), ECategoryMenuType::Section) };
+		return Categories;
+	}
+	// UAssetDefinition End
+};
diff --git a/Engine/Source/Editor/VirtualTexturingEditor/Private/AssetDefinition_VirtualTextureBuilder.h b/Engine/Source/Editor/VirtualTexturingEditor/Private/AssetDefinition_VirtualTextureBuilder.h
new file mode 100644
--- /dev/null
+++ b/Engine/Source/Editor/VirtualTexturingEditor/Private/AssetDefinition_VirtualTextureBuilder.h
@@ -0,0 +1,26 @@
+// Copyright Epic Games, Inc. All Rights Reserved.
+
+#pragma once
+#include "AssetDefinitionDefault.h"
+#include "VT/VirtualTextureBuilder.h"
+#include "AssetDefinition_VirtualTextureBuilder.generated.h"
+
+UCLASS()
+class UAssetDefinition_VirtualTextureBuilder : public UAssetDefinitionDefault
+{
+	GENERATED_BODY()
+public:
+	// UAssetDefinition Begin
+	virtual FText GetAssetDisplayName() const override { return NSLOCTEXT("AssetDefinition", "AssetDefinition_VirtualTextureBuilder", "Streaming Runtime Virtual Texture"); }
+	virtual TSoftClassPtr<UObject> GetAssetClass() const override { return UVirtualTextureBuilder::StaticClass(); }
+	virtual FLinearColor GetAssetColor() const override { return FColor(128, 128, 128); }
+	virtual TConstArrayView<FAssetCategoryPath> GetAssetCategories() const override
+	{
+		static const auto Categories = 
+		{
+			FAssetCategoryPath(EAssetCategoryPaths::Texture, NSLOCTEXT("AssetDefinition", "AssetDefinition_VirtualTextureBuilderSubMenu", "Virtual Texture"), ECategoryMenuType::Section)
+		};
+		return Categories;
+	}
+	// UAssetDefinition End
+};
diff --git a/Engine/Source/Editor/VirtualTexturingEditor/Private/MaterialCacheVirtualTextureTagAssetTypeActions.cpp b/Engine/Source/Editor/VirtualTexturingEditor/Private/MaterialCacheVirtualTextureTagAssetTypeActions.cpp
deleted file mode 100644
diff --git a/Engine/Source/Editor/VirtualTexturingEditor/Private/RuntimeVirtualTextureAssetTypeActions.cpp b/Engine/Source/Editor/VirtualTexturingEditor/Private/RuntimeVirtualTextureAssetTypeActions.cpp
deleted file mode 100644
diff --git a/Engine/Source/Editor/VirtualTexturingEditor/Private/VirtualTextureBuilderAssetTypeActions.cpp b/Engine/Source/Editor/VirtualTexturingEditor/Private/VirtualTextureBuilderAssetTypeActions.cpp
deleted file mode 100644
diff --git a/Engine/Source/Editor/VirtualTexturingEditor/Private/VirtualTexturingEditorModule.cpp b/Engine/Source/Editor/VirtualTexturingEditor/Private/VirtualTexturingEditorModule.cpp
--- a/Engine/Source/Editor/VirtualTexturingEditor/Private/VirtualTexturingEditorModule.cpp
+++ b/Engine/Source/Editor/VirtualTexturingEditor/Private/VirtualTexturingEditorModule.cpp
@@ -4,17 +4,14 @@
 
 #include "Components/RuntimeVirtualTextureComponent.h"
 #include "IPlacementModeModule.h"
-#include "MaterialCacheVirtualTextureTagAssetTypeActions.h"
 #include "MeshPaintVirtualTextureThumbnailRenderer.h"
 #include "PropertyEditorModule.h"
-#include "RuntimeVirtualTextureAssetTypeActions.h"
 #include "RuntimeVirtualTextureBuildStreamingMips.h"
 #include "RuntimeVirtualTextureDetailsCustomization.h"
 #include "RuntimeVirtualTextureThumbnailRenderer.h"
 #include "SConvertToVirtualTexture.h"
 #include "ThumbnailRendering/ThumbnailManager.h"
 #include "UObject/UObjectIterator.h"
-#include "VirtualTextureBuilderAssetTypeActions.h"
 #include "VirtualTextureBuilderThumbnailRenderer.h"
 #include "VT/MeshPaintVirtualTexture.h"
 #include "VT/RuntimeVirtualTexture.h"
@@ -55,11 +52,6 @@ IMPLEMENT_MODULE(FVirtualTexturingEditorModule, VirtualTexturingEditor);
 
 void FVirtualTexturingEditorModule::StartupModule()
 {
-	IAssetTools& AssetTools = FModuleManager::LoadModuleChecked<FAssetToolsModule>("AssetTools").Get();
-	AssetTools.RegisterAssetTypeActions(MakeShareable(new FAssetTypeActions_RuntimeVirtualTexture));
-	AssetTools.RegisterAssetTypeActions(MakeShareable(new FAssetTypeActions_VirtualTextureBuilder));
-	AssetTools.RegisterAssetTypeActions(MakeShareable(new FAssetTypeActions_MaterialCacheVirtualTextureTag));
-
 	FPropertyEditorModule& PropertyModule = FModuleManager::GetModuleChecked<FPropertyEditorModule>("PropertyEditor");
 	PropertyModule.RegisterCustomClassLayout("RuntimeVirtualTexture", FOnGetDetailCustomizationInstance::CreateStatic(&FRuntimeVirtualTextureDetailsCustomization::MakeInstance));
 	PropertyModule.RegisterCustomClassLayout("RuntimeVirtualTextureComponent", FOnGetDetailCustomizationInstance::CreateStatic(&FRuntimeVirtualTextureComponentDetailsCustomization::MakeInstance));
diff --git a/Engine/Source/Editor/VirtualTexturingEditor/VirtualTexturingEditor.Build.cs b/Engine/Source/Editor/VirtualTexturingEditor/VirtualTexturingEditor.Build.cs
--- a/Engine/Source/Editor/VirtualTexturingEditor/VirtualTexturingEditor.Build.cs
+++ b/Engine/Source/Editor/VirtualTexturingEditor/VirtualTexturingEditor.Build.cs
@@ -9,6 +9,7 @@ public class VirtualTexturingEditor : ModuleRules
 	{
 		PrivateDependencyModuleNames.AddRange(
 			new string[] {
+				"AssetDefinition",
 				"AppFramework",
 				"AssetRegistry",
 				"ContentBrowser",
@@ -27,6 +28,7 @@ public class VirtualTexturingEditor : ModuleRules
 				"RHI",
 				"Slate",
 				"SlateCore",
+				"ToolMenus",
 				"UnrealEd",
 			}
 		);
